ColorManagement: Moving the OpenColorIO wrapper module to the engine.

OpenColorIOLib will now be linked against in editor builds. Module config is created lazily only if needed. To be used by subsequent TextureCompressor integration CL.

diff --git a/Engine/Plugins/Compositing/OpenColorIO/OpenColorIO.uplugin b/Engine/Plugins/Compositing/OpenColorIO/OpenColorIO.uplugin
--- a/Engine/Plugins/Compositing/OpenColorIO/OpenColorIO.uplugin
+++ b/Engine/Plugins/Compositing/OpenColorIO/OpenColorIO.uplugin
@@ -20,11 +20,6 @@
 			"Type": "Runtime",
 			"LoadingPhase": "PostConfigInit"
 		},
-		{
-			"Name": "OpenColorIOWrapper",
-			"Type": "Editor",
-			"LoadingPhase": "PostConfigInit"
-		},
 		{
 			"Name": "OpenColorIOEditor",
 			"Type": "Editor",
diff --git a/Engine/Plugins/Compositing/OpenColorIO/Source/OpenColorIOWrapper/OpenColorIOWrapper.Build.cs b/Engine/Plugins/Compositing/OpenColorIO/Source/OpenColorIOWrapper/OpenColorIOWrapper.Build.cs
deleted file mode 100644
diff --git a/Engine/Plugins/Compositing/OpenColorIO/Source/OpenColorIOWrapper/Private/OpenColorIOLibHandler.cpp b/Engine/Plugins/Compositing/OpenColorIO/Source/OpenColorIOWrapper/Private/OpenColorIOLibHandler.cpp
deleted file mode 100644
diff --git a/Engine/Plugins/Compositing/OpenColorIO/Source/OpenColorIOWrapper/Private/OpenColorIOLibHandler.h b/Engine/Plugins/Compositing/OpenColorIO/Source/OpenColorIOWrapper/Private/OpenColorIOLibHandler.h
deleted file mode 100644
diff --git a/Engine/Plugins/Compositing/OpenColorIO/Source/OpenColorIOWrapper/Private/OpenColorIOWrapper.cpp b/Engine/Plugins/Compositing/OpenColorIO/Source/OpenColorIOWrapper/Private/OpenColorIOWrapper.cpp
deleted file mode 100644
diff --git a/Engine/Plugins/Compositing/OpenColorIO/Source/OpenColorIOWrapper/Private/OpenColorIOWrapperModule.cpp b/Engine/Plugins/Compositing/OpenColorIO/Source/OpenColorIOWrapper/Private/OpenColorIOWrapperModule.cpp
deleted file mode 100644
diff --git a/Engine/Plugins/Compositing/OpenColorIO/Source/OpenColorIOWrapper/Private/Tests/OpenColorIOWrapperTests.cpp b/Engine/Plugins/Compositing/OpenColorIO/Source/OpenColorIOWrapper/Private/Tests/OpenColorIOWrapperTests.cpp
deleted file mode 100644
diff --git a/Engine/Plugins/Compositing/OpenColorIO/Source/OpenColorIOWrapper/Public/OpenColorIOWrapper.h b/Engine/Plugins/Compositing/OpenColorIO/Source/OpenColorIOWrapper/Public/OpenColorIOWrapper.h
deleted file mode 100644
diff --git a/Engine/Plugins/Compositing/OpenColorIO/Source/OpenColorIOWrapper/Public/OpenColorIOWrapperDefines.h b/Engine/Plugins/Compositing/OpenColorIO/Source/OpenColorIOWrapper/Public/OpenColorIOWrapperDefines.h
deleted file mode 100644
diff --git a/Engine/Plugins/Compositing/OpenColorIO/Source/OpenColorIOWrapper/Public/OpenColorIOWrapperModule.h b/Engine/Plugins/Compositing/OpenColorIO/Source/OpenColorIOWrapper/Public/OpenColorIOWrapperModule.h
deleted file mode 100644
diff --git a/Engine/Source/Developer/OpenColorIOWrapper/OpenColorIOWrapper.Build.cs b/Engine/Source/Developer/OpenColorIOWrapper/OpenColorIOWrapper.Build.cs
new file mode 100644
--- /dev/null
+++ b/Engine/Source/Developer/OpenColorIOWrapper/OpenColorIOWrapper.Build.cs
@@ -0,0 +1,34 @@
+// Copyright Epic Games, Inc. All Rights Reserved.
+
+namespace UnrealBuildTool.Rules
+{
+	public class OpenColorIOWrapper : ModuleRules
+	{
+		public OpenColorIOWrapper(ReadOnlyTargetRules Target) : base(Target)
+		{
+			PCHUsage = PCHUsageMode.UseExplicitOrSharedPCHs;
+
+			PublicIncludePathModuleNames.AddRange(new string[]
+			{
+				"ColorManagement", // for ColorManagementDefines.h
+			});
+
+			PrivateIncludePathModuleNames.AddRange(new string[]
+			{
+				"Engine", // for TextureDefines.h
+			});
+
+			PrivateDependencyModuleNames.AddRange(new string[]
+			{
+				"ColorManagement",
+				"ImageCore",
+				"OpenColorIOLib",
+			});
+
+			PublicDependencyModuleNames.AddRange(new string[]
+			{
+				"Core",
+			});
+		}
+	}
+}
diff --git a/Engine/Source/Developer/OpenColorIOWrapper/Private/OpenColorIOLibHandler.cpp b/Engine/Source/Developer/OpenColorIOWrapper/Private/OpenColorIOLibHandler.cpp
new file mode 100644
--- /dev/null
+++ b/Engine/Source/Developer/OpenColorIOWrapper/Private/OpenColorIOLibHandler.cpp
@@ -0,0 +1,85 @@
+// Copyright Epic Games, Inc. All Rights Reserved.
+
+#include "OpenColorIOLibHandler.h"
+
+#include "HAL/PlatformProcess.h"
+#include "Logging/LogMacros.h"
+#include "Misc/Paths.h"
+#include "OpenColorIOWrapperModule.h"
+
+
+#if PLATFORM_CPU_ARM_FAMILY
+	#if PLATFORM_WINDOWS
+		#define TARGET_ARCH TEXT("ARM64")
+	#elif PLATFORM_LINUX
+		#define TARGET_ARCH TEXT("aarch64-unknown-linux-gnueabi")
+	#endif
+#else
+	#if PLATFORM_WINDOWS
+		#define TARGET_ARCH TEXT("x64")
+	#elif PLATFORM_LINUX
+		#define TARGET_ARCH TEXT("x86_64-unknown-linux-gnu")
+	#endif
+#endif
+
+
+ //~ Static initialization
+ //--------------------------------------------------------------------
+void* FOpenColorIOLibHandler::LibHandle = nullptr;
+
+
+//~ Initialization functions implementation
+//--------------------------------------------------------------------
+bool FOpenColorIOLibHandler::Initialize()
+{
+#if WITH_OCIO && defined(OCIO_DLL_NAME)
+	check(LibHandle == nullptr);
+
+#if PLATFORM_WINDOWS || PLATFORM_LINUX
+	const FString OCIOBinPath = FPaths::Combine(FPaths::EngineDir(), TEXT("Binaries/ThirdParty/OpenColorIO"), FPlatformProcess::GetBinariesSubdirectory(), TARGET_ARCH);
+#else
+	const FString OCIOBinPath = FPaths::Combine(FPaths::EngineDir(), TEXT("Binaries/ThirdParty/OpenColorIO"), FPlatformProcess::GetBinariesSubdirectory());
+#endif
+	const FString DLLPath = FPaths::Combine(OCIOBinPath, TEXT(PREPROCESSOR_TO_STRING(OCIO_DLL_NAME)));
+
+	FPlatformProcess::PushDllDirectory(*OCIOBinPath);
+
+	if (!FPaths::FileExists(DLLPath))
+	{
+		UE_LOG(LogOpenColorIOWrapper, Error, TEXT("Failed to find the OpenColorIO dll. Plug-in will not be functional."));
+		return false;
+	}
+	
+	LibHandle = FPlatformProcess::GetDllHandle(*DLLPath);
+	FPlatformProcess::PopDllDirectory(*OCIOBinPath);
+
+	if (LibHandle == nullptr)
+	{
+		UE_LOG(LogOpenColorIOWrapper, Error, TEXT("Failed to load required library %s. Plug-in will not be functional."), *DLLPath);
+		return false;
+	}
+
+#else
+	return false;
+#endif // WITH_EDITOR
+	return true;
+}
+
+bool FOpenColorIOLibHandler::IsInitialized()
+{
+	return (LibHandle != nullptr);
+}
+
+void FOpenColorIOLibHandler::Shutdown()
+{
+#if WITH_OCIO && defined(OCIO_DLL_NAME)
+	if (LibHandle != nullptr)
+	{
+		FPlatformProcess::FreeDllHandle(LibHandle);
+		LibHandle = nullptr;
+	}
+#endif // WITH_EDITOR
+}
+
+
+
diff --git a/Engine/Source/Developer/OpenColorIOWrapper/Private/OpenColorIOLibHandler.h b/Engine/Source/Developer/OpenColorIOWrapper/Private/OpenColorIOLibHandler.h
new file mode 100644
--- /dev/null
+++ b/Engine/Source/Developer/OpenColorIOWrapper/Private/OpenColorIOLibHandler.h
@@ -0,0 +1,15 @@
+// Copyright Epic Games, Inc. All Rights Reserved.
+
+#pragma once
+
+
+class FOpenColorIOLibHandler
+{
+public:
+	static bool Initialize();
+	static bool IsInitialized();
+	static void Shutdown();
+
+private:
+	static void* LibHandle;
+};
diff --git a/Engine/Source/Developer/OpenColorIOWrapper/Private/OpenColorIOWrapper.cpp b/Engine/Source/Developer/OpenColorIOWrapper/Private/OpenColorIOWrapper.cpp
new file mode 100644
--- /dev/null
+++ b/Engine/Source/Developer/OpenColorIOWrapper/Private/OpenColorIOWrapper.cpp
@@ -0,0 +1,1261 @@
+// Copyright Epic Games, Inc. All Rights Reserved.
+
+#include "OpenColorIOWrapper.h"
+#include "OpenColorIOWrapperModule.h"
+
+#include "Containers/StringView.h"
+#include "Containers/UnrealString.h"
+#include "Containers/Map.h"
+#include "ColorSpace.h"
+#include "Engine/TextureDefines.h"
+#include "ImageCore.h"
+#include "ImageParallelFor.h"
+
+#if WITH_OCIO
+THIRD_PARTY_INCLUDES_START
+#include "OpenColorIO/OpenColorIO.h"
+THIRD_PARTY_INCLUDES_END
+#endif
+
+#define LOCTEXT_NAMESPACE "OpenColorIOWrapper"
+
+#if PLATFORM_EXCEPTIONS_DISABLED
+	#define OCIO_EXCEPTION_HANDLING_TRY()
+	#define OCIO_EXCEPTION_HANDLING_CATCH(Verbosity, Format, ...) (0)
+#else
+	#define OCIO_EXCEPTION_HANDLING_TRY() \
+		try {
+	
+	// NOTE: A string parameter must always be added to the format string since we automatically include the error message.
+	#define OCIO_EXCEPTION_HANDLING_CATCH(Verbosity, Format, ...) \
+		} catch (OCIO_NAMESPACE::Exception& Exc) { \
+			UE_LOG(LogOpenColorIOWrapper, Verbosity, Format, ##__VA_ARGS__, StringCast<TCHAR>(Exc.what()).Get()); \
+		}
+#endif
+
+namespace OpenColorIOWrapper
+{
+	/** OCIO_NAMESPACE::ROLE_INTERCHANGE_SCENE equivalent, since we currently cannot delay-load this definition. */
+	constexpr const ANSICHAR* GetInterchangeName()
+	{
+		return "aces_interchange";
+	}
+
+	// Build routine since there is no FAnsiString
+	TUniquePtr<ANSICHAR[]> MakeAnsiString(const TCHAR* Str)
+	{
+		int32 Num = FPlatformString::ConvertedLength<ANSICHAR>(Str);
+		TUniquePtr<ANSICHAR[]> Ret =  MakeUnique<ANSICHAR[]>(Num);
+		FMemory::Memcpy(Ret.Get(), StringCast<ANSICHAR>(Str).Get(), Num);
+		return Ret;
+	}
+}
+
+const TCHAR* OpenColorIOWrapper::GetVersion()
+{
+#if WITH_OCIO
+	return TEXT(OCIO_VERSION);
+#else
+	return TEXT("");
+#endif // WITH_OCIO
+}
+
+void OpenColorIOWrapper::ClearAllCaches()
+{
+#if WITH_OCIO
+	OCIO_NAMESPACE::ClearAllCaches();
+#endif
+}
+
+struct FOpenColorIOConfigPimpl
+{
+#if WITH_OCIO
+	OCIO_NAMESPACE::ConstConfigRcPtr Config = nullptr;
+#endif // WITH_OCIO
+};
+
+struct FOpenColorIOProcessorPimpl
+{
+#if WITH_OCIO
+	OCIO_NAMESPACE::ConstProcessorRcPtr Processor = nullptr;
+
+	/** Get processor optimization flags. */
+	static OCIO_NAMESPACE::OptimizationFlags GetOptimizationFlags()
+	{
+		return static_cast<OCIO_NAMESPACE::OptimizationFlags>(
+			OCIO_NAMESPACE::OptimizationFlags::OPTIMIZATION_DEFAULT |
+			OCIO_NAMESPACE::OptimizationFlags::OPTIMIZATION_NO_DYNAMIC_PROPERTIES
+		);
+	}
+#endif // WITH_OCIO
+};
+
+struct FOpenColorIOGPUProcessorPimpl
+{
+#if WITH_OCIO
+	OCIO_NAMESPACE::ConstGPUProcessorRcPtr Processor = nullptr;
+	OCIO_NAMESPACE::GpuShaderDescRcPtr ShaderDescription = nullptr;
+#endif // WITH_OCIO
+};
+
+#if WITH_OCIO
+namespace {
+
+	bool IsImageFormatSupported(const FImageView& InImage)
+	{
+		switch (InImage.Format)
+		{
+		case ERawImageFormat::BGRA8:
+		case ERawImageFormat::RGBA16:
+		case ERawImageFormat::RGBA16F:
+		case ERawImageFormat::RGBA32F:
+			return true;
+		default:
+			return false;
+		}
+	}
+
+	OCIO_NAMESPACE::PackedImageDesc GetImageDesc(const FImageView& InImage)
+	{
+		OCIO_NAMESPACE::ChannelOrdering Ordering = OCIO_NAMESPACE::ChannelOrdering::CHANNEL_ORDERING_RGBA;
+		OCIO_NAMESPACE::BitDepth BitDepth = OCIO_NAMESPACE::BitDepth::BIT_DEPTH_UNKNOWN;
+
+		switch (InImage.Format)
+		{
+		case ERawImageFormat::BGRA8:
+			Ordering = OCIO_NAMESPACE::CHANNEL_ORDERING_BGRA;
+			BitDepth = OCIO_NAMESPACE::BIT_DEPTH_UINT8;
+			break;
+		case ERawImageFormat::RGBA16:
+			Ordering = OCIO_NAMESPACE::CHANNEL_ORDERING_RGBA;
+			BitDepth = OCIO_NAMESPACE::BIT_DEPTH_UINT16;
+			break;
+		case ERawImageFormat::RGBA16F:
+			Ordering = OCIO_NAMESPACE::CHANNEL_ORDERING_RGBA;
+			BitDepth = OCIO_NAMESPACE::BIT_DEPTH_F16;
+			break;
+		case ERawImageFormat::RGBA32F:
+			Ordering = OCIO_NAMESPACE::CHANNEL_ORDERING_RGBA;
+			BitDepth = OCIO_NAMESPACE::BIT_DEPTH_F32;
+			break;
+		default:
+
+			// All other cases should have been eliminated by calling IsImageFormatSupported() previously.
+			checkNoEntry();
+		}
+
+		return OCIO_NAMESPACE::PackedImageDesc(
+			InImage.RawData,
+			static_cast<long>(InImage.GetWidth()),
+			static_cast<long>(InImage.GetHeight()),
+			Ordering,
+			BitDepth,
+			OCIO_NAMESPACE::AutoStride,
+			OCIO_NAMESPACE::AutoStride,
+			OCIO_NAMESPACE::AutoStride
+			);
+	};
+}
+#endif // WITH_OCIO
+
+
+FOpenColorIOWrapperConfig::FOpenColorIOWrapperConfig()
+	: Pimpl(MakePimpl<FOpenColorIOConfigPimpl, EPimplPtrMode::DeepCopy>())
+{ }
+
+FOpenColorIOWrapperConfig::FOpenColorIOWrapperConfig(FStringView InFilePath, EAddWorkingColorSpaceOption InOption)
+	: FOpenColorIOWrapperConfig()
+{
+#if WITH_OCIO
+	TRACE_CPUPROFILER_EVENT_SCOPE(FOpenColorIOWrapperConfig::FOpenColorIOWrapperConfig);
+	OCIO_EXCEPTION_HANDLING_TRY();
+
+		using namespace OCIO_NAMESPACE;
+
+		ConstConfigRcPtr NewConfig = Config::CreateFromFile(StringCast<ANSICHAR>(InFilePath.GetData()).Get());
+
+		if (NewConfig != nullptr)
+		{
+			if (InOption == WCS_AsInterchangeSpace)
+			{
+				const TUniquePtr<ANSICHAR[]> AnsiWorkingColorSpaceName = OpenColorIOWrapper::MakeAnsiString(OpenColorIOWrapper::GetWorkingColorSpaceName());
+				ConstColorSpaceRcPtr InterchangeCS = NewConfig->getColorSpace(NewConfig->getCanonicalName(OpenColorIOWrapper::GetInterchangeName()));
+
+				// When the aces interchange color space is present, we add the working color space as an additional option.
+				if (InterchangeCS != nullptr && NewConfig->getColorSpace(AnsiWorkingColorSpaceName.Get()) == nullptr)
+				{
+					ColorSpaceRcPtr WorkingCS = InterchangeCS->createEditableCopy();
+					WorkingCS->setName(AnsiWorkingColorSpaceName.Get());
+					WorkingCS->setFamily("UE");
+					WorkingCS->clearAliases();
+
+					ConfigRcPtr NewConfigCopy = NewConfig->createEditableCopy();
+					NewConfigCopy->addColorSpace(WorkingCS);
+
+					NewConfig = NewConfigCopy;
+				}
+			}
+			else
+			{
+				UE_LOG(LogOpenColorIOWrapper, Log, TEXT("Could not add the working color space since the config does not define \"aces_interchange\"."))
+			}
+		}
+
+		Pimpl->Config = NewConfig;
+	
+	OCIO_EXCEPTION_HANDLING_CATCH(Error, TEXT("Could not create OCIO configuration file for %s. Error message: %s"), InFilePath.GetData());
+#else
+	UE_LOG(LogOpenColorIOWrapper, Error, TEXT("The OpenColorIO library is not available, functionality will be disabled."));
+#endif // WITH_OCIO
+}
+
+bool FOpenColorIOWrapperConfig::IsValid() const
+{
+#if WITH_OCIO
+	return Pimpl->Config != nullptr;
+#else
+	return false;
+#endif // WITH_OCIO
+}
+
+int32 FOpenColorIOWrapperConfig::GetNumColorSpaces() const
+{
+#if WITH_OCIO
+	if (IsValid())
+	{
+		return Pimpl->Config->getNumColorSpaces(OCIO_NAMESPACE::SEARCH_REFERENCE_SPACE_ALL, OCIO_NAMESPACE::COLORSPACE_ACTIVE);
+	}
+#endif  // WITH_OCIO
+
+	return 0;
+}
+
+FString FOpenColorIOWrapperConfig::GetColorSpaceName(int32 Index) const
+{
+#if WITH_OCIO
+	if (IsValid())
+	{
+		const char* ColorSpaceName = Pimpl->Config->getColorSpaceNameByIndex(OCIO_NAMESPACE::SEARCH_REFERENCE_SPACE_ALL, OCIO_NAMESPACE::COLORSPACE_ACTIVE, Index);
+
+		return StringCast<TCHAR>(ColorSpaceName).Get();
+	}
+#endif // WITH_OCIO
+
+	return {};
+}
+
+int32 FOpenColorIOWrapperConfig::GetColorSpaceIndex(const TCHAR* InColorSpaceName)
+{
+#if WITH_OCIO
+	if (IsValid())
+	{
+		return Pimpl->Config->getIndexForColorSpace(StringCast<ANSICHAR>(InColorSpaceName).Get());
+	}
+#endif // WITH_OCIO
+
+	return false;
+}
+
+FString FOpenColorIOWrapperConfig::GetColorSpaceFamilyName(const TCHAR* InColorSpaceName) const
+{
+#if WITH_OCIO
+	if (IsValid())
+	{
+		OCIO_NAMESPACE::ConstColorSpaceRcPtr ColorSpace = Pimpl->Config->getColorSpace(StringCast<ANSICHAR>(InColorSpaceName).Get());
+		
+		if (ColorSpace != nullptr)
+		{
+			return StringCast<TCHAR>(ColorSpace->getFamily()).Get();
+		}
+	}
+#endif // WITH_OCIO
+
+	return {};
+}
+
+int32 FOpenColorIOWrapperConfig::GetNumDisplays() const
+{
+#if WITH_OCIO
+	if (IsValid())
+	{
+		return Pimpl->Config->getNumDisplays();
+	}
+#endif // WITH_OCIO
+
+	return 0;
+}
+
+FString FOpenColorIOWrapperConfig::GetDisplayName(int32 Index) const
+{
+#if WITH_OCIO
+	if (IsValid())
+	{
+		const char* DisplayName = Pimpl->Config->getDisplay(Index);
+
+		return StringCast<TCHAR>(DisplayName).Get();
+	}
+#endif // WITH_OCIO
+
+	return {};
+}
+
+int32 FOpenColorIOWrapperConfig::GetNumViews(const TCHAR* InDisplayName) const
+{
+#if WITH_OCIO
+	if (IsValid())
+	{
+		return Pimpl->Config->getNumViews(StringCast<ANSICHAR>(InDisplayName).Get());
+	}
+#endif // WITH_OCIO
+
+	return 0;
+}
+
+FString FOpenColorIOWrapperConfig::GetViewName(const TCHAR* InDisplayName, int32 Index) const
+{
+#if WITH_OCIO
+	if (IsValid())
+	{
+		const char* ViewName = Pimpl->Config->getView(StringCast<ANSICHAR>(InDisplayName).Get(), Index);
+
+		return StringCast<TCHAR>(ViewName).Get();
+	}
+#endif // WITH_OCIO
+
+	return {};
+}
+
+FString FOpenColorIOWrapperConfig::GetDisplayViewTransformName(const TCHAR* InDisplayName, const TCHAR* InViewName) const
+{
+#if WITH_OCIO
+	if (IsValid())
+	{
+		const char* TransformName = Pimpl->Config->getDisplayViewTransformName(StringCast<ANSICHAR>(InDisplayName).Get(), StringCast<ANSICHAR>(InViewName).Get());
+
+		return StringCast<TCHAR>(TransformName).Get();
+	}
+#endif // WITH_OCIO
+
+	return {};
+}
+
+FString FOpenColorIOWrapperConfig::GetCacheID() const
+{
+#if WITH_OCIO
+	if (IsValid())
+	{
+		return StringCast<TCHAR>(Pimpl->Config->getCacheID()).Get();
+	}
+#endif // WITH_OCIO
+
+	return {};
+}
+
+FString FOpenColorIOWrapperConfig::GetDebugString() const
+{
+	TStringBuilder<1024> DebugStringBuilder;
+	
+#if WITH_OCIO
+	if (IsValid())
+	{
+		const OCIO_NAMESPACE::ConstConfigRcPtr& Config = Pimpl->Config;
+		if (Config->getNumColorSpaces() > 0)
+		{
+			DebugStringBuilder.Append(TEXT("** ColorSpaces **\n"));
+			
+			const int32 NumCS = Config->getNumColorSpaces(
+				OCIO_NAMESPACE::SEARCH_REFERENCE_SPACE_ALL,   // Iterate over scene & display color spaces.
+				OCIO_NAMESPACE::COLORSPACE_ALL);              // Iterate over active & inactive color spaces.
+
+			for (int32 IndexCS = 0; IndexCS < NumCS; ++IndexCS)
+			{
+				OCIO_NAMESPACE::ConstColorSpaceRcPtr cs = Config->getColorSpace(Config->getColorSpaceNameByIndex(
+					OCIO_NAMESPACE::SEARCH_REFERENCE_SPACE_ALL,
+					OCIO_NAMESPACE::COLORSPACE_ALL,
+					IndexCS));
+
+				DebugStringBuilder.Append(cs->getName());
+				DebugStringBuilder.Append(TEXT("\n"));
+			}
+
+			DebugStringBuilder.Append(TEXT("** (Display, View) pairs **\n"));
+
+			for (int32 IndexDisplay = 0; IndexDisplay < Config->getNumDisplaysAll(); ++IndexDisplay)
+			{
+				const ANSICHAR* DisplayName = Config->getDisplayAll(IndexDisplay);
+
+				// Iterate over shared views.
+				int32 NumViews = Config->getNumViews(DisplayName);
+				for (int IndexView = 0; IndexView < NumViews; ++IndexView)
+				{
+					const ANSICHAR* ViewName = Config->getView(
+						DisplayName,
+						IndexView);
+
+					DebugStringBuilder.Append(TEXT("("));
+					DebugStringBuilder.Append(DisplayName);
+					DebugStringBuilder.Append(TEXT(", "));
+					DebugStringBuilder.Append(ViewName);
+					DebugStringBuilder.Append(TEXT(")\n"));
+				}
+			}
+		}
+	}
+
+#endif // WITH_OCIO
+
+	return DebugStringBuilder.ToString();
+}
+
+
+FOpenColorIOWrapperEngineConfig::FOpenColorIOWrapperEngineConfig()
+	: FOpenColorIOWrapperConfig()
+{
+#if WITH_OCIO
+	TRACE_CPUPROFILER_EVENT_SCOPE(FOpenColorIOWrapperConfig::CreateWorkingColorSpaceToInterchangeConfig)
+	using namespace OCIO_NAMESPACE;
+	using namespace UE::Color;
+
+	ConfigRcPtr StudioConfig = Config::CreateFromBuiltinConfig("studio-config-v1.0.0_aces-v1.3_ocio-v2.1")->createEditableCopy();
+
+	ColorSpaceRcPtr WCS = ColorSpace::Create();
+	WCS->setName(StringCast<ANSICHAR>(OpenColorIOWrapper::GetWorkingColorSpaceName()).Get());
+	WCS->setBitDepth(BIT_DEPTH_F32);
+	WCS->setEncoding("scene-linear");
+	// We know the scene-referred reference space is ACES2065-1, and hence the correct matrix transform.
+	const FMatrix44d TransformMat = Transpose<double>(FColorSpaceTransform(FColorSpace::GetWorking(), FColorSpace(EColorSpace::ACESAP0)));
+	MatrixTransformRcPtr MatrixTransform = MatrixTransform::Create();
+	MatrixTransform->setMatrix(&TransformMat.M[0][0]);
+	WCS->setTransform(MatrixTransform, COLORSPACE_DIR_TO_REFERENCE);
+
+	StudioConfig->addColorSpace(WCS);
+
+	Pimpl->Config = StudioConfig;
+
+	check(IsValid());
+#endif // WITH_OCIO
+}
+
+FOpenColorIOWrapperProcessor::FOpenColorIOWrapperProcessor(
+	const FOpenColorIOWrapperConfig* InConfig,
+	FStringView InSourceColorSpace,
+	FStringView InDestinationColorSpace,
+	const TMap<FString, FString>& InContextKeyValues)
+	: Pimpl(MakePimpl<FOpenColorIOProcessorPimpl, EPimplPtrMode::DeepCopy>())
+	, OwnerConfig(InConfig)
+	, WorkingColorSpaceTransformType(EOpenColorIOWorkingColorSpaceTransform::None)
+{
+	if (InSourceColorSpace == OpenColorIOWrapper::GetWorkingColorSpaceName())
+	{
+		WorkingColorSpaceTransformType = EOpenColorIOWorkingColorSpaceTransform::Source;
+	}
+	else if (InDestinationColorSpace == OpenColorIOWrapper::GetWorkingColorSpaceName())
+	{
+		WorkingColorSpaceTransformType = EOpenColorIOWorkingColorSpaceTransform::Destination;
+	}
+
+#if WITH_OCIO
+	OCIO_EXCEPTION_HANDLING_TRY();
+		if (OwnerConfig != nullptr && OwnerConfig->IsValid())
+		{
+			const OCIO_NAMESPACE::ConstConfigRcPtr& Config = OwnerConfig->Pimpl->Config;
+			OCIO_NAMESPACE::ContextRcPtr Context = Config->getCurrentContext()->createEditableCopy();
+
+			for (const TPair<FString, FString>& KeyValue : InContextKeyValues)
+			{
+				Context->setStringVar(TCHAR_TO_ANSI(*KeyValue.Key), TCHAR_TO_ANSI(*KeyValue.Value));
+			}
+
+			Pimpl->Processor = Config->getProcessor(
+				Context,
+				StringCast<ANSICHAR>(InSourceColorSpace.GetData()).Get(),
+				StringCast<ANSICHAR>(InDestinationColorSpace.GetData()).Get()
+			);
+		}
+	OCIO_EXCEPTION_HANDLING_CATCH(Log, TEXT("Failed to create processor for [%s, %s]. Error message: %s"), InSourceColorSpace.GetData(), InDestinationColorSpace.GetData());
+#endif // WITH_OCIO
+}
+
+FOpenColorIOWrapperProcessor::FOpenColorIOWrapperProcessor(
+	const FOpenColorIOWrapperConfig* InConfig,
+	FStringView InSourceColorSpace,
+	FStringView InDisplay,
+	FStringView InView,
+	bool bInverseDirection,
+	const TMap<FString, FString>& InContextKeyValues)
+	: Pimpl(MakePimpl<FOpenColorIOProcessorPimpl, EPimplPtrMode::DeepCopy>())
+	, OwnerConfig(InConfig)
+	, WorkingColorSpaceTransformType(EOpenColorIOWorkingColorSpaceTransform::None)
+{
+	if (InSourceColorSpace == OpenColorIOWrapper::GetWorkingColorSpaceName())
+	{
+		if (bInverseDirection)
+		{
+			WorkingColorSpaceTransformType = EOpenColorIOWorkingColorSpaceTransform::Destination;
+		}
+		else
+		{
+			WorkingColorSpaceTransformType = EOpenColorIOWorkingColorSpaceTransform::Source;
+		}
+	}
+
+#if WITH_OCIO
+	OCIO_EXCEPTION_HANDLING_TRY();
+		if (OwnerConfig != nullptr && OwnerConfig->IsValid())
+		{
+			const OCIO_NAMESPACE::ConstConfigRcPtr& Config = OwnerConfig->Pimpl->Config;
+			OCIO_NAMESPACE::ContextRcPtr Context = Config->getCurrentContext()->createEditableCopy();
+
+			for (const TPair<FString, FString>& KeyValue : InContextKeyValues)
+			{
+				Context->setStringVar(TCHAR_TO_ANSI(*KeyValue.Key), TCHAR_TO_ANSI(*KeyValue.Value));
+			}
+
+			Pimpl->Processor = Config->getProcessor(
+				StringCast<ANSICHAR>(InSourceColorSpace.GetData()).Get(),
+				StringCast<ANSICHAR>(InDisplay.GetData()).Get(),
+				StringCast<ANSICHAR>(InView.GetData()).Get(),
+				static_cast<OCIO_NAMESPACE::TransformDirection>(bInverseDirection)
+			);
+		}
+	OCIO_EXCEPTION_HANDLING_CATCH(Log, TEXT("Failed to create processor for [%s, %s, %s, %s]. Error message: %s"), InSourceColorSpace.GetData(), InDisplay.GetData(), InView.GetData(), (bInverseDirection ? TEXT("Inverse") : TEXT("Forward")));
+#endif // WITH_OCIO
+}
+
+FOpenColorIOWrapperProcessor::FOpenColorIOWrapperProcessor(const FOpenColorIOWrapperConfig* InConfig, FStringView InNamedTransform, bool bInverseDirection, const TMap<FString, FString>& InContextKeyValues)
+	: Pimpl(MakePimpl<FOpenColorIOProcessorPimpl, EPimplPtrMode::DeepCopy>())
+	, OwnerConfig(InConfig)
+	, WorkingColorSpaceTransformType(EOpenColorIOWorkingColorSpaceTransform::None)
+{
+#if WITH_OCIO
+	OCIO_EXCEPTION_HANDLING_TRY();
+		if (OwnerConfig != nullptr && OwnerConfig->IsValid())
+		{
+			const OCIO_NAMESPACE::ConstConfigRcPtr& Config = OwnerConfig->Pimpl->Config;
+			OCIO_NAMESPACE::ContextRcPtr Context = Config->getCurrentContext()->createEditableCopy();
+
+			for (const TPair<FString, FString>& KeyValue : InContextKeyValues)
+			{
+				Context->setStringVar(TCHAR_TO_ANSI(*KeyValue.Key), TCHAR_TO_ANSI(*KeyValue.Value));
+			}
+
+			Pimpl->Processor = Config->getProcessor(
+				StringCast<ANSICHAR>(InNamedTransform.GetData()).Get(),
+				static_cast<OCIO_NAMESPACE::TransformDirection>(bInverseDirection)
+			);
+		}
+	OCIO_EXCEPTION_HANDLING_CATCH(Log, TEXT("Failed to create processor for [%s]. Error message: %s"), InNamedTransform.GetData());
+#endif // WITH_OCIO
+}
+
+FOpenColorIOWrapperProcessor FOpenColorIOWrapperProcessor::CreateTransformToWorkingColorSpace(const FOpenColorIOWrapperSourceColorSettings& InColorSettings)
+{
+	FOpenColorIOWrapperConfig& BuiltInConfig = IOpenColorIOWrapperModule::Get().GetEngineBuiltInConfig();
+
+	const FString TransformName = GetTransformToWorkingColorSpaceName(InColorSettings);
+
+#if WITH_OCIO
+	using namespace OCIO_NAMESPACE;
+	using namespace UE::Color;
+
+	if (BuiltInConfig.Pimpl->Config->getNamedTransform(StringCast<ANSICHAR>(*TransformName).Get()) == nullptr)
+	{
+		GroupTransformRcPtr TransformToWCS = GroupTransform::Create();
+		NamedTransformRcPtr ParentTransform = NamedTransform::Create();
+		ParentTransform->setName(StringCast<ANSICHAR>(*TransformName).Get());
+
+		switch (InColorSettings.EncodingOverride)
+		{
+		case EEncoding::None:
+		case EEncoding::Linear:
+			ParentTransform->setEncoding("scene-linear");
+			break;
+		case EEncoding::sRGB:
+		{
+			ExponentWithLinearTransformRcPtr ChildTransform = ExponentWithLinearTransform::Create();
+			ChildTransform->setGamma({ 2.4, 2.4, 2.4, 1.0 });
+			ChildTransform->setOffset({ 0.055, 0.055, 0.055, 0.0 });
+
+			TransformToWCS->appendTransform(ChildTransform);
+			ParentTransform->setEncoding("sdr-video");
+		}
+		break;
+		case EEncoding::ST2084:
+		{
+			BuiltinTransformRcPtr ChildTransform = BuiltinTransform::Create();
+			ChildTransform->setStyle("CURVE - ST-2084_to_LINEAR");
+			TransformToWCS->appendTransform(ChildTransform);
+
+			// By default ocio returns nits/100
+			MatrixTransformRcPtr RescaleTransform = MatrixTransform::Create();
+			FMatrix44d ScaleMatrix = FMatrix44d::Identity.ApplyScale(100.0);
+			RescaleTransform->setMatrix(&ScaleMatrix.M[0][0]);
+
+			TransformToWCS->appendTransform(RescaleTransform);
+			ParentTransform->setEncoding("hdr-video");
+		}
+		break;
+		case EEncoding::Gamma22:
+		{
+			ExponentTransformRcPtr ChildTransform = ExponentTransform::Create();
+			ChildTransform->setValue({ 2.2, 2.2, 2.2, 1.0 });
+			ChildTransform->setNegativeStyle(NEGATIVE_PASS_THRU);
+
+			TransformToWCS->appendTransform(ChildTransform);
+			ParentTransform->setEncoding("sdr-video");
+		}
+		break;
+		case EEncoding::BT1886:
+		{
+			ExponentTransformRcPtr ChildTransform = ExponentTransform::Create();
+			ChildTransform->setValue({ 2.4, 2.4, 2.4, 1.0 });
+			ChildTransform->setNegativeStyle(NEGATIVE_PASS_THRU);
+
+			TransformToWCS->appendTransform(ChildTransform);
+			ParentTransform->setEncoding("sdr-video");
+		}
+		break;
+		case EEncoding::Gamma26:
+		{
+			ExponentTransformRcPtr ChildTransform = ExponentTransform::Create();
+			ChildTransform->setValue({ 2.6, 2.6, 2.6, 1.0 });
+			ChildTransform->setNegativeStyle(NEGATIVE_PASS_THRU);
+
+			TransformToWCS->appendTransform(ChildTransform);
+			ParentTransform->setEncoding("sdr-video");
+		}
+		break;
+		case EEncoding::Cineon:
+		{
+			const double BlackOffset = FGenericPlatformMath::Pow(10.0, (95.0 - 685.0) / 300.0);
+			const double LinSideSlope = 1.0 - BlackOffset;
+			const double LinSideOffset = BlackOffset;
+			static constexpr double LogSideSlope = 300.0 / 1023.0;
+			static constexpr double LogSideOffset = 685.0 / 1023.0;
+			static constexpr double Base = 10.;
+
+			LogAffineTransformRcPtr ChildTransform = LogAffineTransform::Create();
+			ChildTransform->setLinSideSlopeValue({ LinSideSlope, LinSideSlope,LinSideSlope });
+			ChildTransform->setLinSideOffsetValue({ LinSideOffset, LinSideOffset,LinSideOffset });
+			ChildTransform->setLogSideSlopeValue({ LogSideSlope, LogSideSlope, LogSideSlope });
+			ChildTransform->setLogSideOffsetValue({ LogSideOffset, LogSideOffset, LogSideOffset });
+			ChildTransform->setBase(Base);
+			ChildTransform->setDirection(TRANSFORM_DIR_INVERSE);
+
+			TransformToWCS->appendTransform(ChildTransform);
+			ParentTransform->setEncoding("log");
+		}
+		break;
+		case EEncoding::REDLog:
+		{
+			const double BlackOffset = FGenericPlatformMath::Pow(10.0, (0.0 - 1023.0) / 511.0);
+			const double LinSideSlope = 1.0 - BlackOffset;
+			const double LinSideOffset = BlackOffset;
+			static constexpr double LogSideSlope = 511.0 / 1023.0;
+			static constexpr double LogSideOffset = 1.0;
+			static constexpr double Base = 10.;
+
+			LogAffineTransformRcPtr ChildTransform = LogAffineTransform::Create();
+			ChildTransform->setLinSideSlopeValue({ LinSideSlope, LinSideSlope,LinSideSlope });
+			ChildTransform->setLinSideOffsetValue({ LinSideOffset, LinSideOffset,LinSideOffset });
+			ChildTransform->setLogSideSlopeValue({ LogSideSlope, LogSideSlope, LogSideSlope });
+			ChildTransform->setLogSideOffsetValue({ LogSideOffset, LogSideOffset, LogSideOffset });
+			ChildTransform->setBase(Base);
+			ChildTransform->setDirection(TRANSFORM_DIR_INVERSE);
+
+			TransformToWCS->appendTransform(ChildTransform);
+			ParentTransform->setEncoding("log");
+		}
+		break;
+		case EEncoding::REDLog3G10:
+		{
+			static constexpr double LinSideSlope = 155.975327;
+			static constexpr double LinSideOffset = 0.01 * LinSideSlope + 1.0;
+			static constexpr double LogSideSlope = 0.224282;
+			static constexpr double LogSideOffset = 0.0;
+			static constexpr double LinSideBreak = -0.01;
+			static constexpr double Base = 10.;
+
+			LogCameraTransformRcPtr ChildTransform = LogCameraTransform::Create({ LinSideBreak, LinSideBreak,LinSideBreak });
+			ChildTransform->setLinSideSlopeValue({ LinSideSlope, LinSideSlope,LinSideSlope });
+			ChildTransform->setLinSideOffsetValue({ LinSideOffset, LinSideOffset,LinSideOffset });
+			ChildTransform->setLogSideSlopeValue({ LogSideSlope, LogSideSlope, LogSideSlope });
+			ChildTransform->setLogSideOffsetValue({ LogSideOffset, LogSideOffset, LogSideOffset });
+			ChildTransform->setBase(Base);
+			ChildTransform->setDirection(TRANSFORM_DIR_INVERSE);
+
+			TransformToWCS->appendTransform(ChildTransform);
+			ParentTransform->setEncoding("log");
+		}
+		break;
+		case EEncoding::SLog1:
+		{
+			static constexpr double LinSideSlope = 1.0 / 0.9;
+			static constexpr double LinSideOffset = 0.037584;
+			static constexpr double LogSideSlope = 0.432699 * 219.0 * 4.0 / 1023.0;
+			static constexpr double LogSideOffset = ((0.616596 + 0.03) * 219.0 + 16.0) * 4.0 / 1023.0;
+			static constexpr double Base = 10.;
+
+			LogAffineTransformRcPtr ChildTransform = LogAffineTransform::Create();
+			ChildTransform->setLinSideSlopeValue({ LinSideSlope, LinSideSlope,LinSideSlope });
+			ChildTransform->setLinSideOffsetValue({ LinSideOffset, LinSideOffset,LinSideOffset });
+			ChildTransform->setLogSideSlopeValue({ LogSideSlope, LogSideSlope, LogSideSlope });
+			ChildTransform->setLogSideOffsetValue({ LogSideOffset, LogSideOffset, LogSideOffset });
+			ChildTransform->setBase(Base);
+			ChildTransform->setDirection(TRANSFORM_DIR_INVERSE);
+
+			TransformToWCS->appendTransform(ChildTransform);
+			ParentTransform->setEncoding("log");
+		}
+		break;
+		case EEncoding::SLog2:
+		{
+			static constexpr double LinSideSlope = 155.0 / 197.1;
+			static constexpr double LinSideOffset = 0.037584;
+			static constexpr double LogSideSlope = 876.0 * 0.432699 / 1023.0;
+			static constexpr double LogSideOffset = (64.0 + 876.0 * 0.646596) / 1023.0;
+			static constexpr double LinSideBreak = 0.0;
+			static constexpr double LinearSlope = 876.0 * (3.53881278538813f / 0.9) / 1023.0;
+			static constexpr double Base = 10.;
+
+			LogCameraTransformRcPtr ChildTransform = LogCameraTransform::Create({ LinSideBreak, LinSideBreak,LinSideBreak });
+			ChildTransform->setLinSideSlopeValue({ LinSideSlope, LinSideSlope,LinSideSlope });
+			ChildTransform->setLinSideOffsetValue({ LinSideOffset, LinSideOffset,LinSideOffset });
+			ChildTransform->setLogSideSlopeValue({ LogSideSlope, LogSideSlope, LogSideSlope });
+			ChildTransform->setLogSideOffsetValue({ LogSideOffset, LogSideOffset, LogSideOffset });
+			ChildTransform->setLinearSlopeValue({ LinearSlope, LinearSlope, LinearSlope });
+			ChildTransform->setBase(Base);
+			ChildTransform->setDirection(TRANSFORM_DIR_INVERSE);
+
+			TransformToWCS->appendTransform(ChildTransform);
+			ParentTransform->setEncoding("log");
+		}
+		break;
+		case EEncoding::SLog3:
+		{
+			static constexpr double LinSideSlope = 5.26315789473684;
+			static constexpr double LinSideOffset = 0.0526315789473684;
+			static constexpr double LogSideSlope = 0.255620723362659;
+			static constexpr double LogSideOffset = 0.410557184750733;
+			static constexpr double LinSideBreak = 0.01125;
+			static constexpr double LinearSlope = 6.62194371177582;
+			static constexpr double Base = 10.;
+
+			LogCameraTransformRcPtr ChildTransform = LogCameraTransform::Create({ LinSideBreak, LinSideBreak,LinSideBreak });
+			ChildTransform->setLinSideSlopeValue({ LinSideSlope, LinSideSlope,LinSideSlope });
+			ChildTransform->setLinSideOffsetValue({ LinSideOffset, LinSideOffset,LinSideOffset });
+			ChildTransform->setLogSideSlopeValue({ LogSideSlope, LogSideSlope, LogSideSlope });
+			ChildTransform->setLogSideOffsetValue({ LogSideOffset, LogSideOffset, LogSideOffset });
+			ChildTransform->setLinearSlopeValue({ LinearSlope, LinearSlope, LinearSlope });
+			ChildTransform->setBase(Base);
+			ChildTransform->setDirection(TRANSFORM_DIR_INVERSE);
+
+			TransformToWCS->appendTransform(ChildTransform);
+			ParentTransform->setEncoding("log");
+		}
+		break;
+		case EEncoding::AlexaV3LogC:
+		{
+			static constexpr double LinSideSlope = 5.55555555555556;
+			static constexpr double LinSideOffset = 0.0522722750251688;
+			static constexpr double LogSideSlope = 0.247189638318671;
+			static constexpr double LogSideOffset = 0.385536998692443;
+			static constexpr double LinSideBreak = 0.0105909904954696;
+			static constexpr double Base = 10.;
+
+			LogCameraTransformRcPtr ChildTransform = LogCameraTransform::Create({ LinSideBreak, LinSideBreak,LinSideBreak });
+			ChildTransform->setLinSideSlopeValue({ LinSideSlope, LinSideSlope,LinSideSlope });
+			ChildTransform->setLinSideOffsetValue({ LinSideOffset, LinSideOffset,LinSideOffset });
+			ChildTransform->setLogSideSlopeValue({ LogSideSlope, LogSideSlope, LogSideSlope });
+			ChildTransform->setLogSideOffsetValue({ LogSideOffset, LogSideOffset, LogSideOffset });
+			ChildTransform->setBase(Base);
+			ChildTransform->setDirection(TRANSFORM_DIR_INVERSE);
+
+			TransformToWCS->appendTransform(ChildTransform);
+			ParentTransform->setEncoding("log");
+		}
+		break;
+		case EEncoding::CanonLog:
+		{
+			static constexpr double LinSideSlope = 10.1596;
+			static constexpr double LinSideOffset = 1.0;
+			static constexpr double LogSideSlope = 0.529136;
+			static constexpr double LogSideOffset = 0.0730597;
+			static constexpr double LinSideBreak = 0.0;
+			static constexpr double Base = 10.;
+
+			LogAffineTransformRcPtr ChildTransform = LogAffineTransform::Create();
+			ChildTransform->setLinSideSlopeValue({ LinSideSlope, LinSideSlope,LinSideSlope });
+			ChildTransform->setLinSideOffsetValue({ LinSideOffset, LinSideOffset,LinSideOffset });
+			ChildTransform->setLogSideSlopeValue({ LogSideSlope, LogSideSlope, LogSideSlope });
+			ChildTransform->setLogSideOffsetValue({ LogSideOffset, LogSideOffset, LogSideOffset });
+			ChildTransform->setBase(Base);
+			ChildTransform->setDirection(TRANSFORM_DIR_INVERSE);
+
+			TransformToWCS->appendTransform(ChildTransform);
+			ParentTransform->setEncoding("log");
+		}
+		break;
+		case EEncoding::ProTune:
+		{
+			static constexpr double LinSideSlope = 112.0;
+			static constexpr double LinSideOffset = 1.0f;
+			const double LogSideSlope = 1.0 / FGenericPlatformMath::Loge(113.0);
+
+			LogAffineTransformRcPtr ChildTransform = LogAffineTransform::Create();
+			ChildTransform->setLinSideSlopeValue({ LinSideSlope, LinSideSlope,LinSideSlope });
+			ChildTransform->setLinSideOffsetValue({ LinSideOffset, LinSideOffset,LinSideOffset });
+			ChildTransform->setLogSideSlopeValue({ LogSideSlope, LogSideSlope, LogSideSlope });
+			ChildTransform->setBase(UE_DOUBLE_EULERS_NUMBER);
+			ChildTransform->setDirection(TRANSFORM_DIR_INVERSE);
+
+			TransformToWCS->appendTransform(ChildTransform);
+			ParentTransform->setEncoding("log");
+		}
+		break;
+		case EEncoding::VLog:
+		{
+			static constexpr double LinSideSlope = 1.0;
+			static constexpr double LinSideOffset = 0.00873;
+			static constexpr double LogSideSlope = 0.241514;
+			static constexpr double LogSideOffset = 0.598206;
+			static constexpr double LinSideBreak = 0.01;
+			static constexpr double Base = 10.;
+			// Note: this is not in the studio config
+			// static constexpr double LinearSlope = 5.6;
+
+			LogCameraTransformRcPtr ChildTransform = LogCameraTransform::Create({ LinSideBreak, LinSideBreak,LinSideBreak });
+			ChildTransform->setLinSideSlopeValue({ LinSideSlope, LinSideSlope,LinSideSlope });
+			ChildTransform->setLinSideOffsetValue({ LinSideOffset, LinSideOffset,LinSideOffset });
+			ChildTransform->setLogSideSlopeValue({ LogSideSlope, LogSideSlope, LogSideSlope });
+			ChildTransform->setLogSideOffsetValue({ LogSideOffset, LogSideOffset, LogSideOffset });
+			//ChildTransform->setLinearSlopeValue({ LinearSlope, LinearSlope, LinearSlope });
+			ChildTransform->setBase(Base);
+			ChildTransform->setDirection(TRANSFORM_DIR_INVERSE);
+
+			TransformToWCS->appendTransform(ChildTransform);
+			ParentTransform->setEncoding("log");
+		}
+		break;
+
+		default:
+			checkNoEntry();
+			break;
+		}
+
+		TOptional<FColorSpace> SourceColorSpace;
+
+		if (InColorSettings.ColorSpaceOverride.IsSet())
+		{
+			const TStaticArray<FVector2d, 4>& SourceChromaticities = InColorSettings.ColorSpaceOverride.GetValue();
+			SourceColorSpace = FColorSpace(
+				SourceChromaticities[0],
+				SourceChromaticities[1],
+				SourceChromaticities[2],
+				SourceChromaticities[3]
+			);
+		}
+		else if (InColorSettings.ColorSpace != EColorSpace::None)
+		{
+			SourceColorSpace = FColorSpace(InColorSettings.ColorSpace);
+		}
+
+		if (SourceColorSpace.IsSet())
+		{
+			if (!SourceColorSpace.GetValue().Equals(FColorSpace::GetWorking()))
+			{
+				const EChromaticAdaptationMethod ChromaticAdapation = static_cast<EChromaticAdaptationMethod>(InColorSettings.ChromaticAdaptationMethod);
+				const FMatrix44d ToWorkingMat = Transpose<double>(FColorSpaceTransform(SourceColorSpace.GetValue(), FColorSpace::GetWorking(), ChromaticAdapation));
+				MatrixTransformRcPtr MatrixTransform = MatrixTransform::Create();
+				MatrixTransform->setMatrix(&ToWorkingMat.M[0][0]);
+				
+				TransformToWCS->appendTransform(MatrixTransform);
+			}
+		}
+
+		ParentTransform->setTransform(TransformToWCS, TRANSFORM_DIR_FORWARD);
+
+		// Update builtin config
+		ConstConfigRcPtr& CurrentConfig = BuiltInConfig.Pimpl->Config;
+		ConfigRcPtr NewConfig = CurrentConfig->createEditableCopy();
+		NewConfig->addNamedTransform(ParentTransform);
+		CurrentConfig = NewConfig;
+	}
+#endif // WITH_OCIO
+
+	return FOpenColorIOWrapperProcessor(&BuiltInConfig, TransformName);
+}
+
+bool FOpenColorIOWrapperProcessor::IsValid() const
+{
+#if WITH_OCIO
+	return OwnerConfig != nullptr && Pimpl->Processor != nullptr;
+#else
+	return false;
+#endif // WITH_OCIO
+}
+
+FString FOpenColorIOWrapperProcessor::GetCacheID() const
+{
+#if WITH_OCIO
+	if (IsValid())
+	{
+		return StringCast<TCHAR>(Pimpl->Processor->getCacheID()).Get();
+	}
+#endif // WITH_OCIO
+
+	return {};
+}
+
+FString FOpenColorIOWrapperProcessor::GetTransformToWorkingColorSpaceName(const FOpenColorIOWrapperSourceColorSettings& InSourceColorSettings)
+{
+	const uint32 SettingsId = (uint32)InSourceColorSettings.EncodingOverride | (uint32)InSourceColorSettings.ColorSpace << 8u | (uint32)InSourceColorSettings.ChromaticAdaptationMethod << 16u;
+	FString TransformName = FString::Printf(TEXT("UE_%u"), SettingsId);
+
+	if (InSourceColorSettings.ColorSpaceOverride.IsSet())
+	{
+		uint32 SrcChromaticityHash = 0;
+		SrcChromaticityHash ^= GetTypeHash(InSourceColorSettings.ColorSpaceOverride.GetValue()[0]);
+		SrcChromaticityHash ^= GetTypeHash(InSourceColorSettings.ColorSpaceOverride.GetValue()[1]);
+		SrcChromaticityHash ^= GetTypeHash(InSourceColorSettings.ColorSpaceOverride.GetValue()[2]);
+		SrcChromaticityHash ^= GetTypeHash(InSourceColorSettings.ColorSpaceOverride.GetValue()[3]);
+		TransformName += FString::Printf(TEXT("_%u"), SrcChromaticityHash);
+	}
+
+	return TransformName;
+}
+
+bool FOpenColorIOWrapperProcessor::TransformColor(FLinearColor& InOutColor) const
+{
+#if WITH_OCIO
+	TRACE_CPUPROFILER_EVENT_SCOPE(FOpenColorIOWrapperProcessor::TransformColor);
+	OCIO_EXCEPTION_HANDLING_TRY();
+
+		if (IsValid())
+		{
+			// Apply the main color transformation
+			OCIO_NAMESPACE::ConstCPUProcessorRcPtr CPUProcessor = Pimpl->Processor->getDefaultCPUProcessor();
+			CPUProcessor->applyRGBA(&InOutColor.R);
+
+			return true;
+		}
+	OCIO_EXCEPTION_HANDLING_CATCH(Error, TEXT("Failed to transform color. Error message: %s"));
+#endif // WITH_OCIO
+
+	return false;
+}
+
+bool FOpenColorIOWrapperProcessor::TransformImage(const FImageView& InOutImage) const
+{
+#if WITH_OCIO
+	TRACE_CPUPROFILER_EVENT_SCOPE(FOpenColorIOWrapperProcessor::TransformImage);
+	OCIO_EXCEPTION_HANDLING_TRY();
+		using namespace OCIO_NAMESPACE;
+
+		if (IsValid())
+		{
+			if (IsImageFormatSupported(InOutImage))
+			{
+				const BitDepth BitDepth = GetImageDesc(InOutImage).getBitDepth();
+
+				// Primary transform processor
+				ConstCPUProcessorRcPtr CPUProcessor = Pimpl->Processor->getOptimizedCPUProcessor(BitDepth, BitDepth, OPTIMIZATION_DEFAULT);
+
+				// Note: This special "interchange to working color space" logic should be removed
+				// once these features are supported by the library.
+				ConstProcessorRcPtr InterchangeProcessor = nullptr;
+				ConstCPUProcessorRcPtr InterchangeCPUProcessor = nullptr;
+				{
+					const TUniquePtr<ANSICHAR[]> AnsiWorkingColorSpaceName = OpenColorIOWrapper::MakeAnsiString(OpenColorIOWrapper::GetWorkingColorSpaceName());
+					ConstConfigRcPtr InterchangeConfig = IOpenColorIOWrapperModule::Get().GetEngineBuiltInConfig().Pimpl->Config;
+					ConstConfigRcPtr Config = OwnerConfig->Pimpl->Config;
+
+					if (WorkingColorSpaceTransformType == EOpenColorIOWorkingColorSpaceTransform::Source)
+					{
+						InterchangeProcessor = InterchangeConfig->GetProcessorFromConfigs(
+							InterchangeConfig,
+							AnsiWorkingColorSpaceName.Get(),
+							Config,
+							Config->getCanonicalName(OpenColorIOWrapper::GetInterchangeName()));
+					}
+					else if (WorkingColorSpaceTransformType == EOpenColorIOWorkingColorSpaceTransform::Destination)
+					{
+						InterchangeProcessor = InterchangeConfig->GetProcessorFromConfigs(
+							Config,
+							Config->getCanonicalName(OpenColorIOWrapper::GetInterchangeName()),
+							InterchangeConfig,
+							AnsiWorkingColorSpaceName.Get());
+					}
+
+					if (InterchangeProcessor)
+					{
+						InterchangeCPUProcessor = InterchangeProcessor->getOptimizedCPUProcessor(BitDepth, BitDepth, OPTIMIZATION_DEFAULT);
+					}
+				}
+
+				// Apply parallelized color transformation
+				FImageCore::ImageParallelFor(TEXT("FOpenColorIOWrapperProcessor.TransformImage.PF"), InOutImage, [&](FImageView& ImagePart)
+					{
+						OCIO_NAMESPACE::PackedImageDesc ImagePartDesc = GetImageDesc(ImagePart);
+
+						// Conditionally apply a conversion from the working color space to interchange space
+						if (WorkingColorSpaceTransformType == EOpenColorIOWorkingColorSpaceTransform::Source)
+						{
+							InterchangeCPUProcessor->apply(ImagePartDesc);
+						}
+
+						// Apply the main color transformation
+						CPUProcessor->apply(ImagePartDesc);
+
+						// Conditionally apply a conversion from the interchange space to the working color space
+						if (WorkingColorSpaceTransformType == EOpenColorIOWorkingColorSpaceTransform::Destination)
+						{
+							InterchangeCPUProcessor->apply(ImagePartDesc);
+						}
+					});
+
+				return true;
+			}
+			else
+			{
+				UE_LOG(LogOpenColorIOWrapper, Warning, TEXT("Unsupported texture format."));
+			}
+		}
+	OCIO_EXCEPTION_HANDLING_CATCH(Error, TEXT("Failed to transform image. Error message: %s"));
+#endif // WITH_OCIO
+
+	return false;
+}
+
+bool FOpenColorIOWrapperProcessor::TransformImage(const FImageView& SrcImage, const FImageView& DestImage) const
+{
+#if WITH_OCIO
+	TRACE_CPUPROFILER_EVENT_SCOPE(FOpenColorIOWrapperProcessor::TransformImage);
+	OCIO_EXCEPTION_HANDLING_TRY();
+		using namespace OCIO_NAMESPACE;
+
+		if (IsValid())
+		{
+			if (IsImageFormatSupported(SrcImage) && IsImageFormatSupported(DestImage))
+			{
+				PackedImageDesc SrcImageDesc = GetImageDesc(SrcImage);
+				PackedImageDesc DestImageDesc = GetImageDesc(DestImage);
+
+				const TUniquePtr<ANSICHAR[]> AnsiWorkingColorSpaceName = OpenColorIOWrapper::MakeAnsiString(OpenColorIOWrapper::GetWorkingColorSpaceName());
+				ConstConfigRcPtr	InterchangeConfig = IOpenColorIOWrapperModule::Get().GetEngineBuiltInConfig().Pimpl->Config;
+				ConstConfigRcPtr	Config = OwnerConfig->Pimpl->Config;
+
+				BitDepth SrcBitDepth = SrcImageDesc.getBitDepth();
+				BitDepth DestBitDepth = DestImageDesc.getBitDepth();
+
+				// Conditionally apply a conversion from the working color space to interchange space
+				if (WorkingColorSpaceTransformType == EOpenColorIOWorkingColorSpaceTransform::Source)
+				{
+					ConstProcessorRcPtr	InterchangeProcessor = InterchangeConfig->GetProcessorFromConfigs(
+						InterchangeConfig,
+						AnsiWorkingColorSpaceName.Get(),
+						Config,
+						Config->getCanonicalName(OpenColorIOWrapper::GetInterchangeName()));
+
+					ConstCPUProcessorRcPtr InterchangeCPUProcessor = InterchangeProcessor->getOptimizedCPUProcessor(SrcBitDepth, SrcBitDepth, OPTIMIZATION_DEFAULT);
+					InterchangeCPUProcessor->apply(SrcImageDesc);
+				}
+
+
+				// Apply the main color transformation
+				ConstCPUProcessorRcPtr CPUProcessor = Pimpl->Processor->getOptimizedCPUProcessor(SrcBitDepth, DestBitDepth, OPTIMIZATION_DEFAULT);
+				CPUProcessor->apply(SrcImageDesc, DestImageDesc);
+
+				// Conditionally apply a conversion from the interchange space to the working color space
+				if (WorkingColorSpaceTransformType == EOpenColorIOWorkingColorSpaceTransform::Destination)
+				{
+					ConstProcessorRcPtr	InterchangeProcessor = InterchangeConfig->GetProcessorFromConfigs(
+						Config,
+						Config->getCanonicalName(OpenColorIOWrapper::GetInterchangeName()),
+						InterchangeConfig,
+						AnsiWorkingColorSpaceName.Get());
+
+					ConstCPUProcessorRcPtr InterchangeCPUProcessor = InterchangeProcessor->getOptimizedCPUProcessor(DestBitDepth, DestBitDepth, OPTIMIZATION_DEFAULT);
+					InterchangeCPUProcessor->apply(DestImageDesc);
+				}
+
+				return true;
+			}
+			else
+			{
+				UE_LOG(LogOpenColorIOWrapper, Warning, TEXT("Unsupported texture format(s)."));
+			}
+		}
+	OCIO_EXCEPTION_HANDLING_CATCH(Error, TEXT("Failed to transform image. Error message: %s"));
+#endif // WITH_OCIO
+
+	return false;
+}
+
+
+
+FOpenColorIOWrapperGPUProcessor::FOpenColorIOWrapperGPUProcessor(FOpenColorIOWrapperProcessor InProcessor, bool bUseLegacy)
+	: ParentProcessor(MoveTemp(InProcessor))
+	, GPUPimpl(MakePimpl<FOpenColorIOGPUProcessorPimpl, EPimplPtrMode::DeepCopy>())
+{
+#if WITH_OCIO
+	TRACE_CPUPROFILER_EVENT_SCOPE(FOpenColorIOWrapperGPUProcessor::FOpenColorIOWrapperGPUProcessor);
+	OCIO_EXCEPTION_HANDLING_TRY();
+		using namespace OCIO_NAMESPACE;
+
+		if (ParentProcessor.IsValid())
+		{
+			GpuShaderDescRcPtr ShaderDescription = GpuShaderDesc::CreateShaderDesc();
+			ShaderDescription->setLanguage(GPU_LANGUAGE_HLSL_DX11);
+			ShaderDescription->setFunctionName(StringCast<ANSICHAR>(OpenColorIOWrapper::GetShaderFunctionName()).Get());
+			ShaderDescription->setResourcePrefix("Ocio");
+
+			ConstGPUProcessorRcPtr GPUProcessor = nullptr;
+			OptimizationFlags OptFlags = FOpenColorIOProcessorPimpl::GetOptimizationFlags();
+
+			if (bUseLegacy)
+			{
+				unsigned int EdgeLength = static_cast<unsigned int>(OpenColorIOWrapper::Legacy3dEdgeLength);
+				GPUProcessor = ParentProcessor.Pimpl->Processor->getOptimizedLegacyGPUProcessor(OptFlags, EdgeLength);
+			}
+			else
+			{
+				GPUProcessor = ParentProcessor.Pimpl->Processor->getOptimizedGPUProcessor(FOpenColorIOProcessorPimpl::GetOptimizationFlags());
+			}
+			GPUProcessor->extractGpuShaderInfo(ShaderDescription);
+
+			GPUPimpl->Processor = GPUProcessor;
+			GPUPimpl->ShaderDescription = ShaderDescription;
+		}
+	OCIO_EXCEPTION_HANDLING_CATCH(Log, TEXT("Failed to fetch shader info for color transform. Error message: %s"));
+#endif // WITH_OCIO
+}
+
+bool FOpenColorIOWrapperGPUProcessor::IsValid() const
+{
+#if WITH_OCIO
+	return ParentProcessor.IsValid() && GPUPimpl->Processor != nullptr && GPUPimpl->ShaderDescription != nullptr;
+#else
+	return false;
+#endif // WITH_OCIO
+}
+
+bool FOpenColorIOWrapperGPUProcessor::GetShader(FString& OutShaderCacheID, FString& OutShaderCode) const
+{
+#if WITH_OCIO
+	if (IsValid())
+	{
+		ensureMsgf(GPUPimpl->ShaderDescription->getNumDynamicProperties() == 0, TEXT("We do not currently support dynamic properties."));
+
+		OutShaderCacheID = StringCast<TCHAR>(GPUPimpl->ShaderDescription->getCacheID()).Get();
+		OutShaderCode = StringCast<TCHAR>(GPUPimpl->ShaderDescription->getShaderText()).Get();
+
+		return true;
+	}
+#endif // WITH_OCIO
+
+	return false;
+}
+
+uint32 FOpenColorIOWrapperGPUProcessor::GetNum3DTextures() const
+{
+#if WITH_OCIO
+	if (IsValid())
+	{
+		return GPUPimpl->ShaderDescription->getNum3DTextures();
+	}
+#endif // WITH_OCIO
+
+	return 0;
+}
+
+bool FOpenColorIOWrapperGPUProcessor::Get3DTexture(uint32 InIndex, FName& OutName, uint32& OutEdgeLength, TextureFilter& OutTextureFilter, const float*& OutData) const
+{
+#if WITH_OCIO
+	ensure(InIndex < GetNum3DTextures());
+	OCIO_EXCEPTION_HANDLING_TRY();
+		if (IsValid())
+		{
+			const ANSICHAR* TextureName = nullptr;
+			const ANSICHAR* SamplerName = nullptr;
+			OCIO_NAMESPACE::Interpolation Interpolation = OCIO_NAMESPACE::INTERP_TETRAHEDRAL;
+
+			// Read texture information
+			GPUPimpl->ShaderDescription->get3DTexture(InIndex, TextureName, SamplerName, OutEdgeLength, Interpolation);
+
+			// Read texture data
+			OutData = 0x0;
+			GPUPimpl->ShaderDescription->get3DTextureValues(InIndex, OutData);
+
+			OutName = FName(TextureName);
+
+			OutTextureFilter = TF_Bilinear;
+			if (Interpolation == OCIO_NAMESPACE::Interpolation::INTERP_NEAREST || Interpolation == OCIO_NAMESPACE::Interpolation::INTERP_TETRAHEDRAL)
+			{
+				OutTextureFilter = TF_Nearest;
+			}
+
+			return TextureName && OutEdgeLength > 0 && OutData;
+		}
+	OCIO_EXCEPTION_HANDLING_CATCH(Log, TEXT("Failed to fetch 3d texture(s) info for color transform. Error message: %s"));
+#endif // WITH_OCIO
+
+	return false;
+}
+uint32 FOpenColorIOWrapperGPUProcessor::GetNumTextures() const
+{
+#if WITH_OCIO
+	if (IsValid())
+	{
+		return GPUPimpl->ShaderDescription->getNumTextures();  //noexcept
+	}
+#endif // WITH_OCIO
+
+	return 0;
+}
+
+bool FOpenColorIOWrapperGPUProcessor::GetTexture(uint32 InIndex, FName& OutName, uint32& OutWidth, uint32& OutHeight, TextureFilter& OutTextureFilter, bool& bOutRedChannelOnly, const float*& OutData) const
+{
+	ensure(InIndex < GetNumTextures());
+
+#if WITH_OCIO
+	OCIO_EXCEPTION_HANDLING_TRY();
+		if (IsValid())
+		{
+			const ANSICHAR* TextureName = nullptr;
+			const ANSICHAR* SamplerName = nullptr;
+			OCIO_NAMESPACE::GpuShaderDesc::TextureType Channel = OCIO_NAMESPACE::GpuShaderDesc::TEXTURE_RGB_CHANNEL;
+			OCIO_NAMESPACE::Interpolation Interpolation = OCIO_NAMESPACE::Interpolation::INTERP_LINEAR;
+
+			// Read texture information
+			GPUPimpl->ShaderDescription->getTexture(InIndex, TextureName, SamplerName, OutWidth, OutHeight, Channel, Interpolation);
+
+			// Read texture data
+			OutData = 0x0;
+			GPUPimpl->ShaderDescription->getTextureValues(InIndex, OutData);
+
+			OutName = FName(TextureName);
+			OutTextureFilter = Interpolation == OCIO_NAMESPACE::Interpolation::INTERP_NEAREST ? TF_Nearest : TF_Bilinear;
+			bOutRedChannelOnly = Channel == OCIO_NAMESPACE::GpuShaderCreator::TEXTURE_RED_CHANNEL;
+
+			return TextureName && OutWidth > 0 && OutHeight > 0 && OutData;
+		}
+	OCIO_EXCEPTION_HANDLING_CATCH(Error, TEXT("Failed to fetch texture(s) info for color transform. Error message: %s"));
+#endif // WITH_OCIO
+
+	return false;
+}
+
+FString FOpenColorIOWrapperGPUProcessor::GetCacheID() const
+{
+#if WITH_OCIO
+	if (IsValid())
+	{
+		return StringCast<TCHAR>(GPUPimpl->Processor->getCacheID()).Get();
+	}
+#endif //WITH_OCIO
+
+	return {};
+}
+
+#undef LOCTEXT_NAMESPACE
diff --git a/Engine/Source/Developer/OpenColorIOWrapper/Private/OpenColorIOWrapperModule.cpp b/Engine/Source/Developer/OpenColorIOWrapper/Private/OpenColorIOWrapperModule.cpp
new file mode 100644
--- /dev/null
+++ b/Engine/Source/Developer/OpenColorIOWrapper/Private/OpenColorIOWrapperModule.cpp
@@ -0,0 +1,55 @@
+// Copyright Epic Games, Inc. All Rights Reserved.
+
+#include "OpenColorIOWrapperModule.h"
+
+#include "OpenColorIOWrapper.h"
+#include "OpenColorIOLibHandler.h"
+
+
+DEFINE_LOG_CATEGORY(LogOpenColorIOWrapper);
+
+#define OPENCOLORIOWRAPPER_MODULE_NAME "OpenColorIOWrapper"
+#define LOCTEXT_NAMESPACE "OpenColorIOWrapperModule"
+
+class FOpenColorIOWrapperModule : public IOpenColorIOWrapperModule
+{
+public:
+	//~ Begin IModuleInterface interface
+	virtual void StartupModule() override
+	{
+		FOpenColorIOLibHandler::Initialize();
+	}
+
+	virtual void ShutdownModule() override
+	{
+		EngineBuiltInConfig.Reset();
+
+		FOpenColorIOLibHandler::Shutdown();
+	}
+	//~ End IModuleInterface interface
+
+	//~ Begin IOpenColorIOWrapperModule interface
+	virtual FOpenColorIOWrapperEngineConfig& GetEngineBuiltInConfig() override
+	{
+		FScopeLock Lock(&ConfigCriticalSection);
+
+		if (!EngineBuiltInConfig)
+		{
+			EngineBuiltInConfig = MakeUnique<FOpenColorIOWrapperEngineConfig>();
+		}
+
+		return *EngineBuiltInConfig;
+	}
+	//~ End IOpenColorIOWrapperModule interface
+private:
+
+	// Critical section for the engine config getter.
+	FCriticalSection ConfigCriticalSection;
+
+	// Global engine config using the built-in studio config, lazily allocated.
+	TUniquePtr<FOpenColorIOWrapperEngineConfig> EngineBuiltInConfig;
+};
+
+IMPLEMENT_MODULE(FOpenColorIOWrapperModule, OpenColorIOWrapper);
+
+#undef LOCTEXT_NAMESPACE
diff --git a/Engine/Source/Developer/OpenColorIOWrapper/Private/Tests/OpenColorIOWrapperTests.cpp b/Engine/Source/Developer/OpenColorIOWrapper/Private/Tests/OpenColorIOWrapperTests.cpp
new file mode 100644
--- /dev/null
+++ b/Engine/Source/Developer/OpenColorIOWrapper/Private/Tests/OpenColorIOWrapperTests.cpp
@@ -0,0 +1,102 @@
+// Copyright Epic Games, Inc. All Rights Reserved.
+
+#include "CoreTypes.h"
+#include "ColorSpace.h"
+#include "Logging/LogMacros.h"
+#include "Math/NumericLimits.h"
+#include "Math/UnrealMathUtility.h"
+#include "Misc/AutomationTest.h"
+#include "OpenColorIOWrapper.h"
+#include "TransferFunctions.h"
+
+#if WITH_DEV_AUTOMATION_TESTS
+
+#if WITH_OCIO
+THIRD_PARTY_INCLUDES_START
+#include "OpenColorIO/OpenColorIO.h"
+THIRD_PARTY_INCLUDES_END
+#endif
+
+DEFINE_LOG_CATEGORY_STATIC(LogUnrealOpenColorIOTest, Log, All);
+
+IMPLEMENT_SIMPLE_AUTOMATION_TEST(FOpenColorIOTransferFunctionsTest, "System.OpenColorIO.DecodeToWorkingColorSpace", EAutomationTestFlags::EditorContext | EAutomationTestFlags::EngineFilter)
+bool FOpenColorIOTransferFunctionsTest::RunTest(const FString& Parameters)
+{
+	bool bSuccess = true;
+
+#if WITH_OCIO
+	using namespace OCIO_NAMESPACE;
+	using namespace UE::Color;
+
+	const FLinearColor TestColor = FLinearColor(0.9f, 0.5f, 0.2f, 1.0f);
+
+	for (uint8 TestEncoding = static_cast<uint8>(EEncoding::None); TestEncoding < static_cast<uint8>(EEncoding::Max); ++TestEncoding)
+	{
+		FLinearColor Expected = UE::Color::Decode(static_cast<EEncoding>(TestEncoding), TestColor);
+
+		FOpenColorIOWrapperSourceColorSettings TestSettings;
+		TestSettings.EncodingOverride = static_cast<EEncoding>(TestEncoding);
+
+		FOpenColorIOWrapperProcessor Processor = FOpenColorIOWrapperProcessor::CreateTransformToWorkingColorSpace(TestSettings);
+
+		FLinearColor Actual = TestColor;
+		Processor.TransformColor(Actual);
+
+		// Note: We make the tolerance relative to the values themselves to account for larger values in PQ.
+		const float Tolerance = UE_KINDA_SMALL_NUMBER * 0.5f * (Actual.R + Expected.R);
+
+		if (!Actual.Equals(Expected, Tolerance))
+		{
+			const FString TestNameToPrint = FString::Printf(TEXT("OpenColorIO: %u:%u"), (uint32)TestSettings.EncodingOverride, (uint32)TestSettings.ColorSpace);
+			AddError(FString::Printf(TEXT("Expected '%s' to be %s, but it was %s."), *TestNameToPrint, *Expected.ToString(), *Actual.ToString()), 1);
+			bSuccess = false;
+		}
+	}
+
+
+	// Name collision test
+	int32 Count = 0;
+	TSet<FString> Keys;
+	FRandomStream RandomStream;
+	RandomStream.Initialize(42);
+	for (uint8 TestEncoding = static_cast<uint8>(EEncoding::None); TestEncoding < static_cast<uint8>(EEncoding::Max); ++TestEncoding)
+	{
+		for (uint8 TestColorSpace = static_cast<uint8>(EColorSpace::None); TestColorSpace <= static_cast<uint8>(EColorSpace::Max); ++TestColorSpace)
+		{
+			for (uint8 TestChromatic = 0; TestChromatic < 2; ++TestChromatic)
+			{
+				FOpenColorIOWrapperSourceColorSettings TestSettings;
+				TestSettings.EncodingOverride = static_cast<EEncoding>(TestEncoding);
+
+				if (TestColorSpace < static_cast<uint8>(EColorSpace::Max))
+				{
+					TestSettings.ColorSpace = static_cast<EColorSpace>(TestColorSpace);
+				}
+				else
+				{
+					// We locally use EColorSpace::Max as a test custom color space.
+					TStaticArray<FVector2d, 4> RandomChromaticities;
+					RandomChromaticities[0] = FVector2d(RandomStream.GetUnitVector());
+					RandomChromaticities[1] = FVector2d(RandomStream.GetUnitVector());
+					RandomChromaticities[2] = FVector2d(RandomStream.GetUnitVector());
+					RandomChromaticities[3] = FVector2d(RandomStream.GetUnitVector());
+
+					TestSettings.ColorSpaceOverride = MoveTemp(RandomChromaticities);
+				}
+				
+				TestSettings.ChromaticAdaptationMethod = static_cast<EChromaticAdaptationMethod>(TestChromatic);
+
+				Keys.Add(FOpenColorIOWrapperProcessor::GetTransformToWorkingColorSpaceName(TestSettings));
+				Count++;
+			}
+		}
+	}
+
+	bSuccess &= TestEqual(TEXT("OpenColorIO: Name hash collision test"), Keys.Num(), Count);
+
+#endif
+
+	return bSuccess;
+}
+
+#endif
\ No newline at end of file
diff --git a/Engine/Source/Developer/OpenColorIOWrapper/Public/OpenColorIOWrapper.h b/Engine/Source/Developer/OpenColorIOWrapper/Public/OpenColorIOWrapper.h
new file mode 100644
--- /dev/null
+++ b/Engine/Source/Developer/OpenColorIOWrapper/Public/OpenColorIOWrapper.h
@@ -0,0 +1,318 @@
+// Copyright Epic Games, Inc. All Rights Reserved.
+
+#pragma once
+
+#include "CoreTypes.h"
+#include "Containers/ContainersFwd.h"
+#include "Containers/StaticArray.h"
+#include "Math/Vector.h"
+#include "Templates/PimplPtr.h"
+#include "Templates/UniquePtr.h"
+
+#include "ColorManagementDefines.h"
+#include "OpenColorIOWrapperDefines.h"
+
+struct FImageView;
+enum TextureFilter : int;
+
+namespace OpenColorIOWrapper
+{
+	/** Color space name of the engine's working color space we insert in OpenColorIO configs. */
+	constexpr const TCHAR* GetWorkingColorSpaceName()
+	{
+		return TEXT("Working Color Space");
+	}
+
+	/** Default generated shader function name. */
+	constexpr const TCHAR* GetShaderFunctionName()
+	{
+		return TEXT("OCIOConvert");
+	}
+
+	/** Default LUT size used in the legacy gpu processor */
+	static constexpr uint32 Legacy3dEdgeLength = 65;
+
+	/** Get the OpenColorIO version string. */
+	OPENCOLORIOWRAPPER_API const TCHAR* GetVersion();
+
+	/** Calls the native function of the same name. */
+	OPENCOLORIOWRAPPER_API void ClearAllCaches();
+}
+
+
+/**
+ * Wrapper around library-native config object.
+ */
+class FOpenColorIOWrapperConfig
+{
+public:
+
+	/** Config initialization options. */
+	enum EAddWorkingColorSpaceOption
+	{
+		WCS_None,
+		// Add the working color space as ACES 2065-1 so that we can target it with an extra conversion from our config if it defines aces_interchange.
+		WCS_AsInterchangeSpace
+	};
+
+	/** Constructor. */
+	FOpenColorIOWrapperConfig();
+
+	/**
+	* Constructor.
+	* 
+	* @param InFilePath Config absolute file path.
+	* @param InOptions Initialization options.
+	*/
+	OPENCOLORIOWRAPPER_API explicit FOpenColorIOWrapperConfig(FStringView InFilePath, EAddWorkingColorSpaceOption InOption = WCS_None);
+
+	OPENCOLORIOWRAPPER_API virtual ~FOpenColorIOWrapperConfig() = default;
+
+	/** Valid when the native config has been successfully created and isn't null. */
+	OPENCOLORIOWRAPPER_API bool IsValid() const;
+	
+	/** Get the number of color spaces in the configuration. */
+	OPENCOLORIOWRAPPER_API int32 GetNumColorSpaces() const;
+	
+	/** Get a color space name at an index. */
+	OPENCOLORIOWRAPPER_API FString GetColorSpaceName(int32 Index) const;
+	
+	/** Get the index of a color space, -1 if none. */
+	OPENCOLORIOWRAPPER_API int32 GetColorSpaceIndex(const TCHAR* InColorSpaceName);
+	
+	/** Get the family name for a color space. */
+	OPENCOLORIOWRAPPER_API FString GetColorSpaceFamilyName(const TCHAR* InColorSpaceName) const;
+	
+	/** Get the number of displays in the configuration. */
+	OPENCOLORIOWRAPPER_API int32 GetNumDisplays() const;
+	
+	/** Get a display name at an index. */
+	OPENCOLORIOWRAPPER_API FString GetDisplayName(int32 Index) const;
+	
+	/** Get the number of views for a display. */
+	OPENCOLORIOWRAPPER_API int32 GetNumViews(const TCHAR* InDisplayName) const;
+	
+	/** Get a view name for its display and index. */
+	OPENCOLORIOWRAPPER_API FString GetViewName(const TCHAR* InDisplayName, int32 Index) const;
+
+	/** Get a display-view trnasform name. */
+	OPENCOLORIOWRAPPER_API FString GetDisplayViewTransformName(const TCHAR* InDisplayName, const TCHAR* InViewName) const;
+
+	/** Get the string hash of the config. */
+	OPENCOLORIOWRAPPER_API FString GetCacheID() const;
+
+	/** Returns a config debug string. */
+	OPENCOLORIOWRAPPER_API FString GetDebugString() const;
+
+protected:
+
+	TPimplPtr<struct FOpenColorIOConfigPimpl, EPimplPtrMode::DeepCopy> Pimpl;
+
+	friend class FOpenColorIOWrapperProcessor;
+	friend class FOpenColorIOWrapperGPUProcessor;
+};
+
+/* Matches FTextureSourceColorSettings in Engine/Texture.h */
+struct FOpenColorIOWrapperSourceColorSettings
+{
+	/** Source encoding of the texture, exposing more options than just sRGB. */
+	UE::Color::EEncoding EncodingOverride = UE::Color::EEncoding::None;
+
+	/** Source color space of the texture. Leave to None for custom color space. */
+	UE::Color::EColorSpace ColorSpace = UE::Color::EColorSpace::None;
+
+	/** Optional chromaticity coordinates of a custom source color space. */
+	TOptional<TStaticArray<FVector2d, 4>> ColorSpaceOverride = {};
+
+	/** Chromatic adaption method applied if the source white point differs from the working color space white point. */
+	UE::Color::EChromaticAdaptationMethod ChromaticAdaptationMethod = UE::Color::EChromaticAdaptationMethod::None;
+};
+
+/**
+ * Wrapper around library-native config object for module-specific use cases.
+ * This uses the built-in studio config as a base.
+ */
+class FOpenColorIOWrapperEngineConfig : public FOpenColorIOWrapperConfig
+{
+public:
+	/**
+	* Constructor.
+	*/
+	FOpenColorIOWrapperEngineConfig();
+
+	FOpenColorIOWrapperEngineConfig(FStringView InFilePath, EAddWorkingColorSpaceOption InOption = WCS_None) = delete;
+};
+
+/**
+ * Wrapper around library-native processor object.
+ * This class can also be used for CPU transforms. It will allocate CPU processors as needed, and reuse them since they are by default cached by the library.
+ */
+class FOpenColorIOWrapperProcessor final
+{
+public:
+	/**
+	* Constructor.
+	*
+	* @param InConfig Owner config.
+	* @param InSourceColorSpace Source color space name.
+	* @param InDestinationColorSpace Destination color space name.
+	* @param InContextKeyValues (Optional) Additional context modifiers.
+	*/
+	OPENCOLORIOWRAPPER_API explicit FOpenColorIOWrapperProcessor(
+		const FOpenColorIOWrapperConfig* InConfig,
+		FStringView InSourceColorSpace,
+		FStringView InDestinationColorSpace,
+		const TMap<FString, FString>& InContextKeyValues = {}
+	);
+
+	/**
+	* Constructor.
+	*
+	* @param InConfig Owner config.
+	* @param InSourceColorSpace Source color space name.
+	* @param InDisplay Display name in display-view transform.
+	* @param InView View name in display-view transform.
+	* @param bInverseDirection Flag for inverse transform direction.
+	* @param InContextKeyValues (Optional) Additional context modifiers.
+	*/
+	OPENCOLORIOWRAPPER_API explicit FOpenColorIOWrapperProcessor(
+		const FOpenColorIOWrapperConfig* InConfig,
+		FStringView InSourceColorSpace,
+		FStringView InDisplay,
+		FStringView InView,
+		bool bInverseDirection = false,
+		const TMap<FString, FString>& InContextKeyValues = {}
+	);
+
+	/**
+	* Constructor.
+	*
+	* @param InConfig Owner config.
+	* @param InNamedTransform Transform name.
+	* @param bInverseDirection Flag for inverse transform direction.
+	* @param InContextKeyValues (Optional) Additional context modifiers.
+	*/
+	OPENCOLORIOWRAPPER_API explicit FOpenColorIOWrapperProcessor(
+		const FOpenColorIOWrapperConfig* InConfig,
+		FStringView InNamedTransform,
+		bool bInverseDirection = false,
+		const TMap<FString, FString>& InContextKeyValues = {}
+	);
+
+	/**
+	* Create a processor from source color settings to working color space.
+	* @param InColorSettings Source encoding and colorspace.
+	* @return Processor of the transform
+	*/
+	OPENCOLORIOWRAPPER_API static FOpenColorIOWrapperProcessor CreateTransformToWorkingColorSpace(
+		const FOpenColorIOWrapperSourceColorSettings& InSourceColorSettings
+	);
+
+	/** Valid when the processor has been successfully created and isn't null. */
+	OPENCOLORIOWRAPPER_API bool IsValid() const;
+
+	/** Get the string hash of the processor. */
+	OPENCOLORIOWRAPPER_API FString GetCacheID() const;
+
+	/**
+	* Get the generated transform name from source color settings to working color space.
+	* @param InColorSettings Source encoding and colorspace.
+	* @return Generated name of the transform
+	*/
+	OPENCOLORIOWRAPPER_API static FString GetTransformToWorkingColorSpaceName(const FOpenColorIOWrapperSourceColorSettings& InSourceColorSettings);
+
+	/** Apply the CPU color transform in-place to the specified color. */
+	OPENCOLORIOWRAPPER_API bool TransformColor(FLinearColor& InOutColor) const;
+
+	/** Apply the CPU color transform in-place to the specified image. */
+	OPENCOLORIOWRAPPER_API bool TransformImage(const FImageView& InOutImage) const;
+
+	/** Apply the CPU color transform from the source image to the destination image.
+	 * (The destination FImageView is const but what it points at is not.)
+	 * 
+	 * Note: This function is currently less optimal, as it is not parallelized.
+	 */
+	OPENCOLORIOWRAPPER_API bool TransformImage(const FImageView& SrcImage, const FImageView& DestImage) const;
+
+private:
+
+	TPimplPtr<struct FOpenColorIOProcessorPimpl, EPimplPtrMode::DeepCopy> Pimpl;
+
+	const FOpenColorIOWrapperConfig* OwnerConfig;
+
+	EOpenColorIOWorkingColorSpaceTransform WorkingColorSpaceTransformType;
+
+	friend class FOpenColorIOWrapperGPUProcessor;
+};
+
+/**
+ * Wrapper around library-native GPU processor object.
+ */
+class FOpenColorIOWrapperGPUProcessor final
+{
+public:
+
+	/**
+	* Constructor.
+	*
+	* @param InProcessor Parent processor.
+	* @param InOptions Initialization options.
+	*/
+	OPENCOLORIOWRAPPER_API explicit FOpenColorIOWrapperGPUProcessor(
+		FOpenColorIOWrapperProcessor InProcessor,
+		bool bUseLegacy = false
+	);
+
+	/** Valid when the processor has been successfully created and isn't null. */
+	OPENCOLORIOWRAPPER_API bool IsValid() const;
+
+	/**
+	* Get the generated shader.
+	*
+	* @param OutShaderCacheID Generated shader hash string.
+	* @param OutShaderCode Generated shader text.
+	* @return True when the result is valid.
+	*/
+	OPENCOLORIOWRAPPER_API bool GetShader(FString& OutShaderCacheID, FString& OutShaderCode) const;
+
+	/** Get the number of 3D LUT textures. */
+	OPENCOLORIOWRAPPER_API uint32 GetNum3DTextures() const;
+
+	/**
+	* Get the index-specific 3D LUT texture used by the shader transform.
+	*
+	* @param InIndex Index of the texture.
+	* @param OutName Shader parameter name of the texture.
+	* @param OutEdgeLength Size of the texture.
+	* @param OutTextureFilter Texture filter type.
+	* @param OutData Raw texel data.
+	* @return True when the result is valid.
+	*/
+	OPENCOLORIOWRAPPER_API bool Get3DTexture(uint32 InIndex, FName& OutName, uint32& OutEdgeLength, TextureFilter& OutTextureFilter, const float*& OutData) const;
+
+	/** Get the number of 2D LUT textures. (1D resources are disabled in our case.) */
+	OPENCOLORIOWRAPPER_API uint32 GetNumTextures() const;
+
+	/**
+	* Get the index-specific 2D LUT texture used by the shader transform.
+	*
+	* @param InIndex Index of the texture.
+	* @param OutName Shader parameter name of the texture.
+	* @param OutWidth Width of the texture.
+	* @param OutHeight Height of the texture.
+	* @param OutTextureFilter Texture filter type.
+	* @param bOutRedChannelOnly Flag indicating whether the texture has a single channel or is RGB.
+	* @param OutData Raw texel data.
+	* @return True when the result is valid.
+	*/
+	OPENCOLORIOWRAPPER_API bool GetTexture(uint32 InIndex, FName& OutName, uint32& OutWidth, uint32& OutHeight, TextureFilter& OutTextureFilter, bool& bOutRedChannelOnly, const float*& OutData) const;
+
+	/** Get the string hash of the gpu processor. */
+	OPENCOLORIOWRAPPER_API FString GetCacheID() const;
+
+private:
+
+	const FOpenColorIOWrapperProcessor ParentProcessor;
+
+	TPimplPtr<struct FOpenColorIOGPUProcessorPimpl, EPimplPtrMode::DeepCopy> GPUPimpl;
+};
diff --git a/Engine/Source/Developer/OpenColorIOWrapper/Public/OpenColorIOWrapperDefines.h b/Engine/Source/Developer/OpenColorIOWrapper/Public/OpenColorIOWrapperDefines.h
new file mode 100644
--- /dev/null
+++ b/Engine/Source/Developer/OpenColorIOWrapper/Public/OpenColorIOWrapperDefines.h
@@ -0,0 +1,14 @@
+// Copyright Epic Games, Inc. All Rights Reserved.
+
+#pragma once
+
+#include "HAL/Platform.h"
+
+/** Enum used to indicate whether the working color space should be used as a source or destination. */
+enum class EOpenColorIOWorkingColorSpaceTransform : uint8
+{
+	None = 0,
+	Source = 1,
+	Destination = 2
+};
+
diff --git a/Engine/Source/Developer/OpenColorIOWrapper/Public/OpenColorIOWrapperModule.h b/Engine/Source/Developer/OpenColorIOWrapper/Public/OpenColorIOWrapperModule.h
new file mode 100644
--- /dev/null
+++ b/Engine/Source/Developer/OpenColorIOWrapper/Public/OpenColorIOWrapperModule.h
@@ -0,0 +1,48 @@
+// Copyright Epic Games, Inc. All Rights Reserved.
+
+#pragma once
+
+#include "Modules/ModuleManager.h"
+#include "Logging/LogMacros.h"
+
+class FOpenColorIOWrapperEngineConfig;
+
+DECLARE_LOG_CATEGORY_EXTERN(LogOpenColorIOWrapper, Log, All);
+
+/**
+ * Interface for the OpenColorIO module.
+ */
+class IOpenColorIOWrapperModule : public IModuleInterface
+{
+public:
+	/**
+	 * Singleton-like access to this module's interface.  This is just for convenience!
+	 * Beware of calling this during the shutdown phase, though.  Your module might have been unloaded already.
+	 *
+	 * @return Returns singleton instance, loading the module on demand if needed
+	 */
+	static IOpenColorIOWrapperModule& Get()
+	{
+		static const FName ModuleName = "OpenColorIOWrapper";
+		return FModuleManager::LoadModuleChecked<IOpenColorIOWrapperModule>(ModuleName);
+	}
+
+	/**
+	 * Checks to see if this module is loaded and ready.  It is only valid to call Get() during shutdown if IsAvailable() returns true.
+	 *
+	 * @return True if the module is loaded and ready to use
+	 */
+	static bool IsAvailable()
+	{
+		static const FName ModuleName = "OpenColorIOWrapper";
+		return FModuleManager::Get().IsModuleLoaded(ModuleName);
+	}
+
+	/**
+	 * Returns a minimal dynamically-created native config for engine (working color spaces) conversions.
+	 */
+	virtual FOpenColorIOWrapperEngineConfig& GetEngineBuiltInConfig() = 0;
+
+	/** Virtual destructor */
+	virtual ~IOpenColorIOWrapperModule() = default;
+};
