OCIO: Prevent concurrency issues with global built-in config for parallel texture processing tasks.

diff --git a/Engine/Source/Developer/OpenColorIOWrapper/Private/OpenColorIOWrapper.cpp b/Engine/Source/Developer/OpenColorIOWrapper/Private/OpenColorIOWrapper.cpp
--- a/Engine/Source/Developer/OpenColorIOWrapper/Private/OpenColorIOWrapper.cpp
+++ b/Engine/Source/Developer/OpenColorIOWrapper/Private/OpenColorIOWrapper.cpp
@@ -503,6 +503,9 @@ FOpenColorIOWrapperProcessor FOpenColorIOWrapperProcessor::CreateTransformToWork
 	using namespace OCIO_NAMESPACE;
 	using namespace UE::Color;
 
+	static FCriticalSection BuiltInConfigCriticalSection;
+	FScopeLock Lock(&BuiltInConfigCriticalSection);
+
 	if (BuiltInConfig.Pimpl->Config->getNamedTransform(StringCast<ANSICHAR>(*TransformName).Get()) == nullptr)
 	{
 		GroupTransformRcPtr TransformToWCS = GroupTransform::Create();
