OCIO: Typo fix for chromatic adaptation.

#rnx

diff --git a/Engine/Source/Developer/OpenColorIOWrapper/Private/OpenColorIOWrapper.cpp b/Engine/Source/Developer/OpenColorIOWrapper/Private/OpenColorIOWrapper.cpp
--- a/Engine/Source/Developer/OpenColorIOWrapper/Private/OpenColorIOWrapper.cpp
+++ b/Engine/Source/Developer/OpenColorIOWrapper/Private/OpenColorIOWrapper.cpp
@@ -812,8 +812,8 @@ FOpenColorIOWrapperProcessor FOpenColorIOWrapperProcessor::CreateTransformToWork
 		{
 			if (!SourceColorSpace.GetValue().Equals(FColorSpace::GetWorking()))
 			{
-				const EChromaticAdaptationMethod ChromaticAdapation = static_cast<EChromaticAdaptationMethod>(InColorSettings.ChromaticAdaptationMethod);
-				const FMatrix44d ToWorkingMat = Transpose<double>(FColorSpaceTransform(SourceColorSpace.GetValue(), FColorSpace::GetWorking(), ChromaticAdapation));
+				const EChromaticAdaptationMethod ChromaticAdaptation = static_cast<EChromaticAdaptationMethod>(InColorSettings.ChromaticAdaptationMethod);
+				const FMatrix44d ToWorkingMat = Transpose<double>(FColorSpaceTransform(SourceColorSpace.GetValue(), FColorSpace::GetWorking(), ChromaticAdaptation));
 				MatrixTransformRcPtr MatrixTransform = MatrixTransform::Create();
 				MatrixTransform->setMatrix(&ToWorkingMat.M[0][0]);
 				
