OCIO: Remove custom working color space transform matrix handling.

This is now possible to due our addition of the missing GetProcessorFromConfigs method for display-views.

#rnx

diff --git a/Engine/Plugins/Compositing/OpenColorIO/Shaders/Private/OpenColorIOShader.usf b/Engine/Plugins/Compositing/OpenColorIO/Shaders/Private/OpenColorIOShader.usf
--- a/Engine/Plugins/Compositing/OpenColorIO/Shaders/Private/OpenColorIOShader.usf
+++ b/Engine/Plugins/Compositing/OpenColorIO/Shaders/Private/OpenColorIOShader.usf
@@ -10,8 +10,6 @@
 =============================================================================*/
 Texture2D		InputTexture;
 SamplerState	InputTextureSampler;
-float4x4		WorkingColorSpaceToInterchange;
-float4x4		InterchangeToWorkingColorSpace;
 float			Gamma;
 uint TransformAlpha;
 
@@ -74,17 +72,9 @@ void MainPS(in noperspective float4 InUVAndScreenPos : TEXCOORD0,
 		InputColor = Unpremultiply(InputColor);
 	}
 
-#if WORKING_COLOR_SPACE_TRANSFORM_TYPE == 1 // EOpenColorIOWorkingColorSpaceTransform::Source
-	InputColor.xyz = mul((float3x3)WorkingColorSpaceToInterchange, InputColor.xyz);
-#endif
-
 	// Apply main transform
 	OutColor = float4(OCIOConvert(InputColor).xyz, InputColor.w);
 
-#if WORKING_COLOR_SPACE_TRANSFORM_TYPE == 2 // EOpenColorIOWorkingColorSpaceTransform::Destination
-	OutColor.xyz = mul((float3x3)InterchangeToWorkingColorSpace, OutColor.xyz);
-#endif
-
 	if(TransformAlpha > 0)
 	{
 		OutColor = Premultiply(OutColor);
diff --git a/Engine/Plugins/Compositing/OpenColorIO/Source/OpenColorIO/Private/OpenColorIOColorTransform.cpp b/Engine/Plugins/Compositing/OpenColorIO/Source/OpenColorIO/Private/OpenColorIOColorTransform.cpp
--- a/Engine/Plugins/Compositing/OpenColorIO/Source/OpenColorIO/Private/OpenColorIOColorTransform.cpp
+++ b/Engine/Plugins/Compositing/OpenColorIO/Source/OpenColorIO/Private/OpenColorIOColorTransform.cpp
@@ -9,7 +9,6 @@
 #include "OpenColorIOConfiguration.h"
 #include "OpenColorIOModule.h"
 #include "OpenColorIOSettings.h"
-#include "OpenColorIOWrapperDefines.h"
 #include "UObject/UE5MainStreamObjectVersion.h"
 #include "UObject/UObjectIterator.h"
 #include "DataDrivenShaderPlatformInfo.h"
@@ -133,7 +132,6 @@ void UOpenColorIOColorTransform::GetOpenColorIOLUTKeyGuid(const FString& InProce
 UOpenColorIOColorTransform::UOpenColorIOColorTransform(const FObjectInitializer& ObjectInitializer)
 	: Super(ObjectInitializer)
 {
-	WorkingColorSpaceTransformType = EOpenColorIOWorkingColorSpaceTransform::None;
 }
 
 bool UOpenColorIOColorTransform::Initialize(UOpenColorIOConfiguration* InOwner, const FString& InSourceColorSpace, const FString& InDestinationColorSpace, const TMap<FString, FString>& InContextKeyValues)
@@ -191,15 +189,6 @@ bool UOpenColorIOColorTransform::Initialize(const FString& InSourceColorSpace, c
 		return false;
 	}
 
-	if (InSourceColorSpace == OpenColorIOWrapper::GetWorkingColorSpaceName())
-	{
-		WorkingColorSpaceTransformType = EOpenColorIOWorkingColorSpaceTransform::Source;
-	}
-	else if (InDestinationColorSpace == OpenColorIOWrapper::GetWorkingColorSpaceName())
-	{
-		WorkingColorSpaceTransformType = EOpenColorIOWorkingColorSpaceTransform::Destination;
-	}
-
 	SourceColorSpace = InSourceColorSpace;
 	DestinationColorSpace = InDestinationColorSpace;
 	bIsDisplayViewType = false;
@@ -218,18 +207,6 @@ bool UOpenColorIOColorTransform::Initialize(const FString& InSourceColorSpace, c
 		return false;
 	}
 
-	if (InSourceColorSpace == OpenColorIOWrapper::GetWorkingColorSpaceName())
-	{
-		if (InDirection == EOpenColorIOViewTransformDirection::Forward)
-		{
-			WorkingColorSpaceTransformType = EOpenColorIOWorkingColorSpaceTransform::Source;
-		}
-		else
-		{
-			WorkingColorSpaceTransformType = EOpenColorIOWorkingColorSpaceTransform::Destination;
-		}
-	}
-
 	SourceColorSpace = InSourceColorSpace;
 	DestinationColorSpace.Empty();
 	Display = InDisplay;
@@ -251,7 +228,7 @@ void UOpenColorIOColorTransform::CacheResourceShadersForCooking(EShaderPlatform
 	FOpenColorIOTransformResource* NewResource = new FOpenColorIOTransformResource();
 
 	FName AssetPath = GetOutermost()->GetFName();
-	NewResource->SetupResource((ERHIFeatureLevel::Type)TargetFeatureLevel, InShaderHash, InShaderCode, InRawConfigHash, GetTransformFriendlyName(), AssetPath, WorkingColorSpaceTransformType);
+	NewResource->SetupResource((ERHIFeatureLevel::Type)TargetFeatureLevel, InShaderHash, InShaderCode, InRawConfigHash, GetTransformFriendlyName(), AssetPath);
 
 	const bool bApplyCompletedShaderMap = false;
 	const bool bIsCooking = true;
@@ -468,7 +445,7 @@ void UOpenColorIOColorTransform::CacheResourceShadersForRendering(bool bRegenera
 				RawConfigHash = ConfigurationOwner->GetConfigWrapper()->GetCacheID();
 				AssetPath = GetOutermost()->GetFName();
 #endif // WITH_EDITOR
-				TransformResource->SetupResource(CacheFeatureLevel, ShaderCodeHash, ShaderCode, RawConfigHash, GetTransformFriendlyName(), AssetPath, WorkingColorSpaceTransformType);
+				TransformResource->SetupResource(CacheFeatureLevel, ShaderCodeHash, ShaderCode, RawConfigHash, GetTransformFriendlyName(), AssetPath);
 
 				const bool bApplyCompletedShaderMap = true;
 
diff --git a/Engine/Plugins/Compositing/OpenColorIO/Source/OpenColorIO/Private/OpenColorIOConfiguration.cpp b/Engine/Plugins/Compositing/OpenColorIO/Source/OpenColorIO/Private/OpenColorIOConfiguration.cpp
--- a/Engine/Plugins/Compositing/OpenColorIO/Source/OpenColorIO/Private/OpenColorIOConfiguration.cpp
+++ b/Engine/Plugins/Compositing/OpenColorIO/Source/OpenColorIO/Private/OpenColorIOConfiguration.cpp
@@ -2,6 +2,7 @@
 
 #include "OpenColorIOConfiguration.h"
 
+#include "ColorSpace.h"
 #include "Containers/SortedMap.h"
 #include "EngineAnalytics.h"
 #include "Engine/VolumeTexture.h"
@@ -185,6 +186,18 @@ void UOpenColorIOConfiguration::ReloadExistingColorspaces(bool bForce)
 			LoadedConfigHash += FString(TEXT("_Inv"));
 		}
 
+		const UE::Color::FColorSpace& WCS = UE::Color::FColorSpace::GetWorking();
+		if (!WCS.IsSRGB())
+		{
+			// The working color space is uniquely defined by its chromaticities (as loaded from renderer settings).
+			uint32 WCSHash = 0;
+			WCSHash ^= GetTypeHash(WCS.GetRedChromaticity());
+			WCSHash ^= GetTypeHash(WCS.GetGreenChromaticity());
+			WCSHash ^= GetTypeHash(WCS.GetBlueChromaticity());
+			WCSHash ^= GetTypeHash(WCS.GetWhiteChromaticity());
+			LoadedConfigHash += FString::Printf(TEXT("_WCS-%u"), WCSHash);
+		}
+
 		// Hash is different, proceed with the regeneration...
 		if (ConfigHash != LoadedConfigHash || bForce)
 		{
diff --git a/Engine/Plugins/Compositing/OpenColorIO/Source/OpenColorIO/Private/OpenColorIORendering.cpp b/Engine/Plugins/Compositing/OpenColorIO/Source/OpenColorIO/Private/OpenColorIORendering.cpp
--- a/Engine/Plugins/Compositing/OpenColorIO/Source/OpenColorIO/Private/OpenColorIORendering.cpp
+++ b/Engine/Plugins/Compositing/OpenColorIO/Source/OpenColorIO/Private/OpenColorIORendering.cpp
@@ -2,7 +2,6 @@
 
 #include "OpenColorIORendering.h"
 
-#include "ColorSpace.h"
 #include "Engine/RendererSettings.h"
 #include "Engine/TextureRenderTarget2D.h"
 #include "Engine/World.h"
@@ -17,7 +16,6 @@
 #include "OpenColorIOShaderType.h"
 #include "OpenColorIOShared.h"
 #include "OpenColorIOColorTransform.h"
-#include "OpenColorIOWrapperDefines.h"
 #include "SceneInterface.h"
 #include "ScreenPass.h"
 #include "TextureResource.h"
@@ -25,23 +23,6 @@
 
 float FOpenColorIORendering::DefaultDisplayGamma = 2.2f;
 
-namespace {
-	using namespace UE::Color;
-
-	// Static local storage to prevent color space recomputation every frame. This is viable since a WCS change requires a project relaunch.
-	const FMatrix44f& GetWorkingColorSpaceToInterchangeTransform()
-	{
-		static FMatrix44f Transform = Transpose<float>(FColorSpaceTransform(FColorSpace::GetWorking(), FColorSpace(EColorSpace::ACESAP0)));
-		return Transform;
-	}
-
-	const FMatrix44f& GetInterchangeToWorkingColorSpaceTransform()
-	{
-		static FMatrix44f Transform = Transpose<float>(FColorSpaceTransform(FColorSpace(EColorSpace::ACESAP0), FColorSpace::GetWorking()));
-		return Transform;
-	}
-}
-
 // static
 void FOpenColorIORendering::AddPass_RenderThread(
 	FRDGBuilder& GraphBuilder,
@@ -69,21 +50,6 @@ void FOpenColorIORendering::AddPass_RenderThread(
 		Parameters->InputTexture = Input.Texture;
 		Parameters->InputTextureSampler = TStaticSamplerState<>::GetRHI();
 		OpenColorIOBindTextureResources(Parameters, InPassResource.TextureResources);
-
-		// Apply a transform between the working color space and the interchange color space, if necessary.
-		switch (InPassResource.ShaderResource->GetWorkingColorSpaceTransformType())
-		{
-		case EOpenColorIOWorkingColorSpaceTransform::Source:
-			Parameters->WorkingColorSpaceToInterchange = GetWorkingColorSpaceToInterchangeTransform();
-			break;
-		case EOpenColorIOWorkingColorSpaceTransform::Destination:
-			Parameters->InterchangeToWorkingColorSpace = GetInterchangeToWorkingColorSpaceTransform();
-			break;
-
-		default:
-			// do nothing, shader parameter is unused.
-			break;	
-		}
 		Parameters->Gamma = InGamma;
 		Parameters->TransformAlpha = (uint32)TransformAlpha;
 		Parameters->RenderTargets[0] = Output.GetRenderTargetBinding();
diff --git a/Engine/Plugins/Compositing/OpenColorIO/Source/OpenColorIO/Private/OpenColorIOShaderMap.cpp b/Engine/Plugins/Compositing/OpenColorIO/Source/OpenColorIO/Private/OpenColorIOShaderMap.cpp
--- a/Engine/Plugins/Compositing/OpenColorIO/Source/OpenColorIO/Private/OpenColorIOShaderMap.cpp
+++ b/Engine/Plugins/Compositing/OpenColorIO/Source/OpenColorIO/Private/OpenColorIOShaderMap.cpp
@@ -208,8 +208,6 @@ FShader* FOpenColorIOShaderType::FinishCompileShader(
 void FOpenColorIOShaderType::SetupCompileEnvironment(EShaderPlatform InPlatform, const FOpenColorIOTransformResource* InColorTransform, FShaderCompilerEnvironment& OutEnvironment) const
 {
 	ModifyCompilationEnvironment(FOpenColorIOShaderPermutationParameters(GetFName(), InPlatform, InColorTransform), OutEnvironment);
-
-	OutEnvironment.SetDefine(TEXT("WORKING_COLOR_SPACE_TRANSFORM_TYPE"), (int32)InColorTransform->GetWorkingColorSpaceTransformType());
 }
 
 #endif // WITH_EDITOR
diff --git a/Engine/Plugins/Compositing/OpenColorIO/Source/OpenColorIO/Private/OpenColorIOShared.cpp b/Engine/Plugins/Compositing/OpenColorIO/Source/OpenColorIO/Private/OpenColorIOShared.cpp
--- a/Engine/Plugins/Compositing/OpenColorIO/Source/OpenColorIO/Private/OpenColorIOShared.cpp
+++ b/Engine/Plugins/Compositing/OpenColorIO/Source/OpenColorIO/Private/OpenColorIOShared.cpp
@@ -13,7 +13,6 @@
 #include "OpenColorIOShaderType.h"
 #include "OpenColorIOShader.h"
 #include "OpenColorIOShaderCompilationManager.h"
-#include "OpenColorIOWrapperDefines.h"
 #include "RendererInterface.h"
 #include "ShaderCompiler.h"
 #include "Stats/StatsMisc.h"
@@ -34,7 +33,6 @@ FOpenColorIOTransformResource::FOpenColorIOTransformResource()
 	, FeatureLevel(ERHIFeatureLevel::SM5)
 	, bContainsInlineShaders(false)
 	, bLoadedCookedShaderMapId(false)
-	, WorkingColorSpaceTransformType(EOpenColorIOWorkingColorSpaceTransform::None)
 {}
 
 FOpenColorIOTransformResource::~FOpenColorIOTransformResource()
@@ -211,7 +209,12 @@ void FOpenColorIOTransformResource::SerializeShaderMap(FArchive& Ar)
 	}
 }
 
-void FOpenColorIOTransformResource::SetupResource(ERHIFeatureLevel::Type InFeatureLevel, const FString& InShaderCodeHash, const FString& InShadercode, const FString& InRawConfigHash, const FString& InFriendlyName, const FName& InAssetPath, EOpenColorIOWorkingColorSpaceTransform InWorkingColorSpaceTransformType)
+void FOpenColorIOTransformResource::SetupResource(ERHIFeatureLevel::Type InFeatureLevel, const FString& InShaderCodeHash, const FString& InShadercode, const FString& InRawConfigHash, const FString& InFriendlyName, const FName& InAssetPath, EOpenColorIOWorkingColorSpaceTransform /*InWorkingColorSpaceTransformType*/)
+{
+	SetupResource(InFeatureLevel, InShadercode, InShadercode, InRawConfigHash, InFriendlyName, InAssetPath);
+}
+
+void FOpenColorIOTransformResource::SetupResource(ERHIFeatureLevel::Type InFeatureLevel, const FString& InShaderCodeHash, const FString& InShadercode, const FString& InRawConfigHash, const FString& InFriendlyName, const FName& InAssetPath)
 {
 	// When this happens we assume that shader was cooked and we don't need to do anything. 
 	if (InShaderCodeHash.IsEmpty() && InRawConfigHash.IsEmpty())
@@ -219,15 +222,13 @@ void FOpenColorIOTransformResource::SetupResource(ERHIFeatureLevel::Type InFeatu
 		return;
 	}
 
-
-	const FString ConcatHashes = InShaderCodeHash + InRawConfigHash + FString::FromInt((int32)InWorkingColorSpaceTransformType);
+	const FString ConcatHashes = InShaderCodeHash + InRawConfigHash;
 	FSHAHash FullHash;
 	FSHA1::HashBuffer(TCHAR_TO_ANSI(*(ConcatHashes)), ConcatHashes.Len(), FullHash.Hash);
 
 	ShaderCodeAndConfigHash = FullHash.ToString();
 	ShaderCode = InShadercode;
 	FriendlyName = InFriendlyName.Replace(TEXT("/"), TEXT(""));
-	WorkingColorSpaceTransformType = InWorkingColorSpaceTransformType;
 	AssetPath = InAssetPath;
 
 	SetFeatureLevel(InFeatureLevel);
diff --git a/Engine/Plugins/Compositing/OpenColorIO/Source/OpenColorIO/Public/OpenColorIOColorTransform.h b/Engine/Plugins/Compositing/OpenColorIO/Source/OpenColorIO/Public/OpenColorIOColorTransform.h
--- a/Engine/Plugins/Compositing/OpenColorIO/Source/OpenColorIO/Public/OpenColorIOColorTransform.h
+++ b/Engine/Plugins/Compositing/OpenColorIO/Source/OpenColorIO/Public/OpenColorIOColorTransform.h
@@ -262,9 +262,6 @@ private:
 	/** Key-value string pairs used to define the processor's context. */
 	TMap<FString, FString> ContextKeyValues;
 
-	/** Enum used to indicate whether the working color space should be used as a source or destination. */
-	EOpenColorIOWorkingColorSpaceTransform WorkingColorSpaceTransformType;
-
 	FRenderCommandFence ReleaseFence;
 
 #if WITH_EDITORONLY_DATA
diff --git a/Engine/Plugins/Compositing/OpenColorIO/Source/OpenColorIO/Public/OpenColorIOShader.h b/Engine/Plugins/Compositing/OpenColorIO/Source/OpenColorIO/Public/OpenColorIOShader.h
--- a/Engine/Plugins/Compositing/OpenColorIO/Source/OpenColorIO/Public/OpenColorIOShader.h
+++ b/Engine/Plugins/Compositing/OpenColorIO/Source/OpenColorIO/Public/OpenColorIOShader.h
@@ -20,8 +20,6 @@ class FRHITexture;
 BEGIN_SHADER_PARAMETER_STRUCT(FOpenColorIOPixelShaderParameters, )
 	SHADER_PARAMETER_RDG_TEXTURE(Texture2D, InputTexture)
 	SHADER_PARAMETER_SAMPLER(SamplerState, InputTextureSampler)
-	SHADER_PARAMETER(FMatrix44f, WorkingColorSpaceToInterchange)
-	SHADER_PARAMETER(FMatrix44f, InterchangeToWorkingColorSpace)
 	SHADER_PARAMETER(float, Gamma)
 	SHADER_PARAMETER(uint32, TransformAlpha)
 	
diff --git a/Engine/Plugins/Compositing/OpenColorIO/Source/OpenColorIO/Public/OpenColorIOShared.h b/Engine/Plugins/Compositing/OpenColorIO/Source/OpenColorIO/Public/OpenColorIOShared.h
--- a/Engine/Plugins/Compositing/OpenColorIO/Source/OpenColorIO/Public/OpenColorIOShared.h
+++ b/Engine/Plugins/Compositing/OpenColorIO/Source/OpenColorIO/Public/OpenColorIOShared.h
@@ -24,7 +24,13 @@
 #include "Templates/RefCounting.h"
 #include "OpenColorIOShaderCompilationManager.h"
 
-enum class EOpenColorIOWorkingColorSpaceTransform : uint8;
+/** Enum used to indicate whether the working color space should be used as a source or destination. */
+enum class EOpenColorIOWorkingColorSpaceTransform : uint8
+{
+	None = 0,
+	Source = 1,
+	Destination = 2
+};
 
 class FOpenColorIOTransformResource;
 class FOpenColorIOShaderMap;
@@ -446,7 +452,7 @@ public:
 
 	const FString& GetFriendlyName()	const { return FriendlyName; }
 
-
+	UE_DEPRECATED(5.4, "SetupResource with EOpenColorIOWorkingColorSpaceTransform deprecated.")
 	void SetupResource(
 		ERHIFeatureLevel::Type InFeatureLevel,
 		const FString& InShaderCodeHash,
@@ -457,6 +463,15 @@ public:
 		EOpenColorIOWorkingColorSpaceTransform InWorkingColorSpaceTransformType
 	);
 
+	void SetupResource(
+		ERHIFeatureLevel::Type InFeatureLevel,
+		const FString& InShaderCodeHash,
+		const FString& InShadercode,
+		const FString& InRawConfigHash,
+		const FString& InFriendlyName,
+		const FName& InAssetPath
+	);
+
 	void SetCompileErrors(TArray<FString> &InErrors)
 	{
 		CompileErrors = InErrors;
@@ -489,7 +504,9 @@ public:
 
 	
 	bool IsSame(const FOpenColorIOShaderMapId& InId) const;
-	EOpenColorIOWorkingColorSpaceTransform GetWorkingColorSpaceTransformType() const { return WorkingColorSpaceTransformType; }
+
+	UE_DEPRECATED(5.4, "GetWorkingColorSpaceTransformType is deprecated.")
+	EOpenColorIOWorkingColorSpaceTransform GetWorkingColorSpaceTransformType() const { return EOpenColorIOWorkingColorSpaceTransform::None; }
 protected:
 #if WITH_EDITOR
 	/**
@@ -546,7 +563,6 @@ private:
 	uint32 bContainsInlineShaders : 1;
 	uint32 bLoadedCookedShaderMapId : 1;
 	FOpenColorIOShaderMapId CookedShaderMapId;
-	EOpenColorIOWorkingColorSpaceTransform WorkingColorSpaceTransformType;
 
 #if WITH_EDITOR
 	/**
diff --git a/Engine/Source/Developer/OpenColorIOWrapper/Private/OpenColorIOWrapper.cpp b/Engine/Source/Developer/OpenColorIOWrapper/Private/OpenColorIOWrapper.cpp
--- a/Engine/Source/Developer/OpenColorIOWrapper/Private/OpenColorIOWrapper.cpp
+++ b/Engine/Source/Developer/OpenColorIOWrapper/Private/OpenColorIOWrapper.cpp
@@ -68,7 +68,6 @@ struct FOpenColorIOConfigPimpl
 struct FOpenColorIOProcessorPimpl
 {
 	OCIO_NAMESPACE::ConstProcessorRcPtr Processor = nullptr;
-	OCIO_NAMESPACE::ConstProcessorRcPtr InterchangeProcessor = nullptr;
 
 	/** Get processor optimization flags. */
 	static OCIO_NAMESPACE::OptimizationFlags GetOptimizationFlags()
@@ -78,29 +77,6 @@ struct FOpenColorIOProcessorPimpl
 			OCIO_NAMESPACE::OptimizationFlags::OPTIMIZATION_NO_DYNAMIC_PROPERTIES
 		);
 	}
-
-	/** Create the interchange processor used for conversion between the working color space and the interchange space. */
-	void SetupInterchangeProcessor(const OCIO_NAMESPACE::ConstConfigRcPtr& InConfig, const OCIO_NAMESPACE::ConstConfigRcPtr& InterchangeConfig, EOpenColorIOWorkingColorSpaceTransform InWorkingColorSpaceTransformType)
-	{
-		const TUniquePtr<ANSICHAR[]> AnsiWorkingColorSpaceName = OpenColorIOWrapper::MakeAnsiString(OpenColorIOWrapper::GetWorkingColorSpaceName());
-
-		if (InWorkingColorSpaceTransformType == EOpenColorIOWorkingColorSpaceTransform::Source)
-		{
-			InterchangeProcessor = InterchangeConfig->GetProcessorFromConfigs(
-				InterchangeConfig,
-				AnsiWorkingColorSpaceName.Get(),
-				InConfig,
-				InConfig->getCanonicalName(OCIO_NAMESPACE::ROLE_INTERCHANGE_SCENE));
-		}
-		else if (InWorkingColorSpaceTransformType == EOpenColorIOWorkingColorSpaceTransform::Destination)
-		{
-			InterchangeProcessor = InterchangeConfig->GetProcessorFromConfigs(
-				InConfig,
-				InConfig->getCanonicalName(OCIO_NAMESPACE::ROLE_INTERCHANGE_SCENE),
-				InterchangeConfig,
-				AnsiWorkingColorSpaceName.Get());
-		}
-	}
 };
 
 struct FOpenColorIOGPUProcessorPimpl
@@ -422,17 +398,7 @@ FOpenColorIOWrapperProcessor::FOpenColorIOWrapperProcessor(
 	FStringView InDestinationColorSpace,
 	const TMap<FString, FString>& InContextKeyValues)
 	: Pimpl(MakePimpl<FOpenColorIOProcessorPimpl, EPimplPtrMode::DeepCopy>())
-	, WorkingColorSpaceTransformType(EOpenColorIOWorkingColorSpaceTransform::None)
 {
-	if (InSourceColorSpace == OpenColorIOWrapper::GetWorkingColorSpaceName())
-	{
-		WorkingColorSpaceTransformType = EOpenColorIOWorkingColorSpaceTransform::Source;
-	}
-	else if (InDestinationColorSpace == OpenColorIOWrapper::GetWorkingColorSpaceName())
-	{
-		WorkingColorSpaceTransformType = EOpenColorIOWorkingColorSpaceTransform::Destination;
-	}
-
 	OCIO_EXCEPTION_HANDLING_TRY();
 		if (InConfig != nullptr && InConfig->IsValid())
 		{
@@ -444,15 +410,36 @@ FOpenColorIOWrapperProcessor::FOpenColorIOWrapperProcessor(
 				Context->setStringVar(TCHAR_TO_ANSI(*KeyValue.Key), TCHAR_TO_ANSI(*KeyValue.Value));
 			}
 
-			Pimpl->Processor = Config->getProcessor(
-				Context,
-				StringCast<ANSICHAR>(InSourceColorSpace.GetData()).Get(),
-				StringCast<ANSICHAR>(InDestinationColorSpace.GetData()).Get()
-			);
+			const TCHAR* WCSColorSpaceName = OpenColorIOWrapper::GetWorkingColorSpaceName();
+			const TUniquePtr<ANSICHAR[]> AnsiWorkingColorSpaceName = OpenColorIOWrapper::MakeAnsiString(WCSColorSpaceName);
 
-			if (WorkingColorSpaceTransformType != EOpenColorIOWorkingColorSpaceTransform::None)
+			if (InSourceColorSpace == WCSColorSpaceName)
+			{
+				Pimpl->Processor = OCIO_NAMESPACE::Config::GetProcessorFromConfigs(
+					Context,
+					IOpenColorIOWrapperModule::Get().GetEngineBuiltInConfig().Pimpl->Config,
+					AnsiWorkingColorSpaceName.Get(),
+					Context,
+					Config,
+					StringCast<ANSICHAR>(InDestinationColorSpace.GetData()).Get());
+			}
+			else if (InDestinationColorSpace == WCSColorSpaceName)
+			{
+				Pimpl->Processor = OCIO_NAMESPACE::Config::GetProcessorFromConfigs(
+					Context,
+					Config,
+					StringCast<ANSICHAR>(InSourceColorSpace.GetData()).Get(),
+					Context,
+					IOpenColorIOWrapperModule::Get().GetEngineBuiltInConfig().Pimpl->Config,
+					AnsiWorkingColorSpaceName.Get());
+			}
+			else
 			{
-				Pimpl->SetupInterchangeProcessor(Config, IOpenColorIOWrapperModule::Get().GetEngineBuiltInConfig().Pimpl->Config, WorkingColorSpaceTransformType);
+				Pimpl->Processor = Config->getProcessor(
+					Context,
+					StringCast<ANSICHAR>(InSourceColorSpace.GetData()).Get(),
+					StringCast<ANSICHAR>(InDestinationColorSpace.GetData()).Get()
+				);
 			}
 		}
 	OCIO_EXCEPTION_HANDLING_CATCH(Log, TEXT("Failed to create processor for [%s, %s]. Error message: %s"), InSourceColorSpace.GetData(), InDestinationColorSpace.GetData());
@@ -466,20 +453,7 @@ FOpenColorIOWrapperProcessor::FOpenColorIOWrapperProcessor(
 	bool bInverseDirection,
 	const TMap<FString, FString>& InContextKeyValues)
 	: Pimpl(MakePimpl<FOpenColorIOProcessorPimpl, EPimplPtrMode::DeepCopy>())
-	, WorkingColorSpaceTransformType(EOpenColorIOWorkingColorSpaceTransform::None)
 {
-	if (InSourceColorSpace == OpenColorIOWrapper::GetWorkingColorSpaceName())
-	{
-		if (bInverseDirection)
-		{
-			WorkingColorSpaceTransformType = EOpenColorIOWorkingColorSpaceTransform::Destination;
-		}
-		else
-		{
-			WorkingColorSpaceTransformType = EOpenColorIOWorkingColorSpaceTransform::Source;
-		}
-	}
-
 	OCIO_EXCEPTION_HANDLING_TRY();
 		if (InConfig != nullptr && InConfig->IsValid())
 		{
@@ -491,16 +465,30 @@ FOpenColorIOWrapperProcessor::FOpenColorIOWrapperProcessor(
 				Context->setStringVar(TCHAR_TO_ANSI(*KeyValue.Key), TCHAR_TO_ANSI(*KeyValue.Value));
 			}
 
-			Pimpl->Processor = Config->getProcessor(
-				StringCast<ANSICHAR>(InSourceColorSpace.GetData()).Get(),
-				StringCast<ANSICHAR>(InDisplay.GetData()).Get(),
-				StringCast<ANSICHAR>(InView.GetData()).Get(),
-				static_cast<OCIO_NAMESPACE::TransformDirection>(bInverseDirection)
-			);
+			const TCHAR* WCSColorSpaceName = OpenColorIOWrapper::GetWorkingColorSpaceName();
+			const TUniquePtr<ANSICHAR[]> AnsiWorkingColorSpaceName = OpenColorIOWrapper::MakeAnsiString(WCSColorSpaceName);
 
-			if (WorkingColorSpaceTransformType != EOpenColorIOWorkingColorSpaceTransform::None)
+			if (InSourceColorSpace == WCSColorSpaceName)
 			{
-				Pimpl->SetupInterchangeProcessor(Config, IOpenColorIOWrapperModule::Get().GetEngineBuiltInConfig().Pimpl->Config, WorkingColorSpaceTransformType);
+				Pimpl->Processor = OCIO_NAMESPACE::Config::GetProcessorFromConfigs(
+					Context,
+					IOpenColorIOWrapperModule::Get().GetEngineBuiltInConfig().Pimpl->Config,
+					AnsiWorkingColorSpaceName.Get(),
+					Context,
+					Config,
+					StringCast<ANSICHAR>(InDisplay.GetData()).Get(),
+					StringCast<ANSICHAR>(InView.GetData()).Get(),
+					static_cast<OCIO_NAMESPACE::TransformDirection>(bInverseDirection));
+			}
+			else
+			{
+				Pimpl->Processor = Config->getProcessor(
+					Context,
+					StringCast<ANSICHAR>(InSourceColorSpace.GetData()).Get(),
+					StringCast<ANSICHAR>(InDisplay.GetData()).Get(),
+					StringCast<ANSICHAR>(InView.GetData()).Get(),
+					static_cast<OCIO_NAMESPACE::TransformDirection>(bInverseDirection)
+				);
 			}
 		}
 	OCIO_EXCEPTION_HANDLING_CATCH(Log, TEXT("Failed to create processor for [%s, %s, %s, %s]. Error message: %s"), InSourceColorSpace.GetData(), InDisplay.GetData(), InView.GetData(), (bInverseDirection ? TEXT("Inverse") : TEXT("Forward")));
@@ -508,7 +496,6 @@ FOpenColorIOWrapperProcessor::FOpenColorIOWrapperProcessor(
 
 FOpenColorIOWrapperProcessor::FOpenColorIOWrapperProcessor(const FOpenColorIOWrapperConfig* InConfig, FStringView InNamedTransform, bool bInverseDirection, const TMap<FString, FString>& InContextKeyValues)
 	: Pimpl(MakePimpl<FOpenColorIOProcessorPimpl, EPimplPtrMode::DeepCopy>())
-	, WorkingColorSpaceTransformType(EOpenColorIOWorkingColorSpaceTransform::None)
 {
 	OCIO_EXCEPTION_HANDLING_TRY();
 		if (InConfig != nullptr && InConfig->IsValid())
@@ -522,6 +509,7 @@ FOpenColorIOWrapperProcessor::FOpenColorIOWrapperProcessor(const FOpenColorIOWra
 			}
 
 			Pimpl->Processor = Config->getProcessor(
+				Context,
 				StringCast<ANSICHAR>(InNamedTransform.GetData()).Get(),
 				static_cast<OCIO_NAMESPACE::TransformDirection>(bInverseDirection)
 			);
@@ -934,36 +922,13 @@ bool FOpenColorIOWrapperProcessor::TransformImage(const FImageView& InOutImage)
 				// Primary transform processor
 				ConstCPUProcessorRcPtr CPUProcessor = Pimpl->Processor->getOptimizedCPUProcessor(BitDepth, BitDepth, OPTIMIZATION_DEFAULT);
 
-				// Note: This special "interchange to working color space" logic should be removed
-				// once these features are supported by the library.
-				ConstCPUProcessorRcPtr InterchangeCPUProcessor = nullptr;
-				if (WorkingColorSpaceTransformType != EOpenColorIOWorkingColorSpaceTransform::None)
-				{
-					if (ensure(Pimpl->InterchangeProcessor))
-					{
-						InterchangeCPUProcessor = Pimpl->InterchangeProcessor->getOptimizedCPUProcessor(BitDepth, BitDepth, OPTIMIZATION_DEFAULT);
-					}
-				}
-
 				// Apply parallelized color transformation
 				FImageCore::ImageParallelFor(TEXT("FOpenColorIOWrapperProcessor.TransformImage.PF"), InOutImage, [&](FImageView& ImagePart)
 					{
 						OCIO_NAMESPACE::PackedImageDesc ImagePartDesc = GetImageDesc(ImagePart);
 
-						// Conditionally apply a conversion from the working color space to interchange space
-						if (WorkingColorSpaceTransformType == EOpenColorIOWorkingColorSpaceTransform::Source)
-						{
-							InterchangeCPUProcessor->apply(ImagePartDesc);
-						}
-
 						// Apply the main color transformation
 						CPUProcessor->apply(ImagePartDesc);
-
-						// Conditionally apply a conversion from the interchange space to the working color space
-						if (WorkingColorSpaceTransformType == EOpenColorIOWorkingColorSpaceTransform::Destination)
-						{
-							InterchangeCPUProcessor->apply(ImagePartDesc);
-						}
 					});
 
 				return true;
@@ -993,30 +958,10 @@ bool FOpenColorIOWrapperProcessor::TransformImage(const FImageView& SrcImage, co
 				BitDepth SrcBitDepth = SrcImageDesc.getBitDepth();
 				BitDepth DestBitDepth = DestImageDesc.getBitDepth();
 
-				// Conditionally apply a conversion from the working color space to interchange space
-				if (WorkingColorSpaceTransformType == EOpenColorIOWorkingColorSpaceTransform::Source)
-				{
-					if (ensure(Pimpl->InterchangeProcessor))
-					{
-						ConstCPUProcessorRcPtr InterchangeCPUProcessor = Pimpl->InterchangeProcessor->getOptimizedCPUProcessor(SrcBitDepth, SrcBitDepth, OPTIMIZATION_DEFAULT);
-						InterchangeCPUProcessor->apply(SrcImageDesc);
-					}
-				}
-
 				// Apply the main color transformation
 				ConstCPUProcessorRcPtr CPUProcessor = Pimpl->Processor->getOptimizedCPUProcessor(SrcBitDepth, DestBitDepth, OPTIMIZATION_DEFAULT);
 				CPUProcessor->apply(SrcImageDesc, DestImageDesc);
 
-				// Conditionally apply a conversion from the interchange space to the working color space
-				if (WorkingColorSpaceTransformType == EOpenColorIOWorkingColorSpaceTransform::Destination)
-				{
-					if (ensure(Pimpl->InterchangeProcessor))
-					{
-						ConstCPUProcessorRcPtr InterchangeCPUProcessor = Pimpl->InterchangeProcessor->getOptimizedCPUProcessor(DestBitDepth, DestBitDepth, OPTIMIZATION_DEFAULT);
-						InterchangeCPUProcessor->apply(DestImageDesc);
-					}
-				}
-
 				return true;
 			}
 			else
diff --git a/Engine/Source/Developer/OpenColorIOWrapper/Public/OpenColorIOWrapper.h b/Engine/Source/Developer/OpenColorIOWrapper/Public/OpenColorIOWrapper.h
--- a/Engine/Source/Developer/OpenColorIOWrapper/Public/OpenColorIOWrapper.h
+++ b/Engine/Source/Developer/OpenColorIOWrapper/Public/OpenColorIOWrapper.h
@@ -12,7 +12,6 @@
 #include "Templates/UniquePtr.h"
 
 #include "ColorManagementDefines.h"
-#include "OpenColorIOWrapperDefines.h"
 
 struct FImageView;
 enum TextureFilter : int;
@@ -249,8 +248,6 @@ private:
 
 	TPimplPtr<struct FOpenColorIOProcessorPimpl, EPimplPtrMode::DeepCopy> Pimpl;
 
-	EOpenColorIOWorkingColorSpaceTransform WorkingColorSpaceTransformType;
-
 	friend class FOpenColorIOWrapperGPUProcessor;
 };
 
diff --git a/Engine/Source/Developer/OpenColorIOWrapper/Public/OpenColorIOWrapperDefines.h b/Engine/Source/Developer/OpenColorIOWrapper/Public/OpenColorIOWrapperDefines.h
deleted file mode 100644
