OCIO: Minor fix to avoid parallel-for loop if the processor is no-op.

#rnx

diff --git a/Engine/Source/Developer/OpenColorIOWrapper/Private/OpenColorIOWrapper.cpp b/Engine/Source/Developer/OpenColorIOWrapper/Private/OpenColorIOWrapper.cpp
--- a/Engine/Source/Developer/OpenColorIOWrapper/Private/OpenColorIOWrapper.cpp
+++ b/Engine/Source/Developer/OpenColorIOWrapper/Private/OpenColorIOWrapper.cpp
@@ -922,14 +922,17 @@ bool FOpenColorIOWrapperProcessor::TransformImage(const FImageView& InOutImage)
 				// Primary transform processor
 				ConstCPUProcessorRcPtr CPUProcessor = Pimpl->Processor->getOptimizedCPUProcessor(BitDepth, BitDepth, OPTIMIZATION_DEFAULT);
 
-				// Apply parallelized color transformation
-				FImageCore::ImageParallelFor(TEXT("FOpenColorIOWrapperProcessor.TransformImage.PF"), InOutImage, [&](FImageView& ImagePart)
-					{
-						OCIO_NAMESPACE::PackedImageDesc ImagePartDesc = GetImageDesc(ImagePart);
-
-						// Apply the main color transformation
-						CPUProcessor->apply(ImagePartDesc);
-					});
+				// Apply parallelized color transformation (when it isn't a no-op)
+				if (!CPUProcessor->isNoOp())
+				{
+					FImageCore::ImageParallelFor(TEXT("FOpenColorIOWrapperProcessor.TransformImage.PF"), InOutImage, [&](FImageView& ImagePart)
+						{
+							OCIO_NAMESPACE::PackedImageDesc ImagePartDesc = GetImageDesc(ImagePart);
+
+							// Apply the main color transformation
+							CPUProcessor->apply(ImagePartDesc);
+						});
+				}
 
 				return true;
 			}
