OCIO: Fix cache id crash on invalid configs.

diff --git a/Engine/Plugins/Compositing/OpenColorIO/Source/OpenColorIO/Private/OpenColorIOConfiguration.cpp b/Engine/Plugins/Compositing/OpenColorIO/Source/OpenColorIO/Private/OpenColorIOConfiguration.cpp
--- a/Engine/Plugins/Compositing/OpenColorIO/Source/OpenColorIO/Private/OpenColorIOConfiguration.cpp
+++ b/Engine/Plugins/Compositing/OpenColorIO/Source/OpenColorIO/Private/OpenColorIOConfiguration.cpp
@@ -178,7 +178,14 @@ void UOpenColorIOConfiguration::ReloadExistingColorspaces(bool bForce)
 
 	if (Config && Config->IsValid())
 	{
-		FString LoadedConfigHash = Config->GetCacheID() + FString(OpenColorIOWrapper::GetVersion());
+		FString LoadedConfigHash = Config->GetCacheID();
+		if (LoadedConfigHash.IsEmpty())
+		{
+			UE_LOG(LogOpenColorIO, Warning, TEXT("Failed to get cache ID: forcing constant transform recreation. Please fix invalid config."));
+			bForce = true;
+		}
+
+		LoadedConfigHash += FString(OpenColorIOWrapper::GetVersion());
 		
 		const UOpenColorIOSettings* Settings = GetDefault<UOpenColorIOSettings>();
 		if (Settings->bSupportInverseViewTransforms)
diff --git a/Engine/Source/Developer/OpenColorIOWrapper/Private/OpenColorIOWrapper.cpp b/Engine/Source/Developer/OpenColorIOWrapper/Private/OpenColorIOWrapper.cpp
--- a/Engine/Source/Developer/OpenColorIOWrapper/Private/OpenColorIOWrapper.cpp
+++ b/Engine/Source/Developer/OpenColorIOWrapper/Private/OpenColorIOWrapper.cpp
@@ -22,6 +22,7 @@ THIRD_PARTY_INCLUDES_END
 #if PLATFORM_EXCEPTIONS_DISABLED
 	#define OCIO_EXCEPTION_HANDLING_TRY()
 	#define OCIO_EXCEPTION_HANDLING_CATCH(Verbosity, Format, ...) (0)
+	#define OCIO_EXCEPTION_HANDLING_CATCH_ERROR()
 #else
 	#define OCIO_EXCEPTION_HANDLING_TRY() \
 		try {
@@ -31,6 +32,8 @@ THIRD_PARTY_INCLUDES_END
 		} catch (OCIO_NAMESPACE::Exception& Exc) { \
 			UE_LOG(LogOpenColorIOWrapper, Verbosity, Format, ##__VA_ARGS__, StringCast<TCHAR>(Exc.what()).Get()); \
 		}
+
+	#define OCIO_EXCEPTION_HANDLING_CATCH_ERROR() OCIO_EXCEPTION_HANDLING_CATCH(Error, TEXT("%s"))
 #endif
 
 namespace OpenColorIOWrapper
@@ -300,10 +303,12 @@ FString FOpenColorIOWrapperConfig::GetDisplayViewTransformName(const TCHAR* InDi
 
 FString FOpenColorIOWrapperConfig::GetCacheID() const
 {
-	if (IsValid())
-	{
-		return StringCast<TCHAR>(Pimpl->Config->getCacheID()).Get();
-	}
+	OCIO_EXCEPTION_HANDLING_TRY();
+		if (IsValid())
+		{
+			return StringCast<TCHAR>(Pimpl->Config->getCacheID()).Get();
+		}
+	OCIO_EXCEPTION_HANDLING_CATCH_ERROR();
 
 	return {};
 }
@@ -863,10 +868,12 @@ bool FOpenColorIOWrapperProcessor::IsValid() const
 
 FString FOpenColorIOWrapperProcessor::GetCacheID() const
 {
-	if (IsValid())
-	{
-		return StringCast<TCHAR>(Pimpl->Processor->getCacheID()).Get();
-	}
+	OCIO_EXCEPTION_HANDLING_TRY();
+		if (IsValid())
+		{
+			return StringCast<TCHAR>(Pimpl->Processor->getCacheID()).Get();
+		}
+	OCIO_EXCEPTION_HANDLING_CATCH_ERROR()
 
 	return {};
 }
@@ -1123,10 +1130,12 @@ bool FOpenColorIOWrapperGPUProcessor::GetTexture(uint32 InIndex, FName& OutName,
 
 FString FOpenColorIOWrapperGPUProcessor::GetCacheID() const
 {
-	if (IsValid())
-	{
-		return StringCast<TCHAR>(GPUPimpl->Processor->getCacheID()).Get();
-	}
+	OCIO_EXCEPTION_HANDLING_TRY();
+		if (IsValid())
+		{
+			return StringCast<TCHAR>(GPUPimpl->Processor->getCacheID()).Get();
+		}
+	OCIO_EXCEPTION_HANDLING_CATCH_ERROR()
 
 	return {};
 }
